Sandboxed Editing: Added Sandboxed Editing plugin, which is responsible for the UX experience in the editor for managing sandboxes. It contains the business logic for directly interacting with sandboxes. For now, there is only a status bar for it but in the future the UI will be made more advanced and better integrated with Multi-User.

#Synced-CL 47637010

diff --git a/Engine/Plugins/Developer/Sandbox/SandboxedEditing/SandboxedEditing.uplugin b/Engine/Plugins/Developer/Sandbox/SandboxedEditing/SandboxedEditing.uplugin
new file mode 100644
--- /dev/null
+++ b/Engine/Plugins/Developer/Sandbox/SandboxedEditing/SandboxedEditing.uplugin
@@ -0,0 +1,31 @@
+{
+	"FileVersion": 3,
+	"Version": 1,
+	"VersionName": "0.1",
+	"FriendlyName": "Sandboxed Editing",
+	"Description": "Allows editing of files in sandbox enviroments. Edited files can be persisted into your main project files.",
+	"Category": "Developer",
+	"CreatedBy": "Epic Games, Inc.",
+	"CreatedByURL": "https://epicgames.com",
+	"DocsURL": "",
+	"MarketplaceURL": "",
+	"SupportURL": "",
+	"EnabledByDefault": false,
+	"CanContainContent": false,
+	"Hidden": true,
+	"IsBetaVersion": true,
+	"Installed": false,
+	"Modules": [
+		{
+			"Name": "SandboxedEditing",
+			"Type": "Editor",
+			"LoadingPhase": "Default"
+		}
+	],
+	"Plugins": [
+		{
+			"Name": "FileSandbox",
+			"Enabled": true
+		}
+	]
+}
diff --git a/Engine/Plugins/Developer/Sandbox/SandboxedEditing/Source/SandboxedEditing/Private/Framework/CommandMappings.cpp b/Engine/Plugins/Developer/Sandbox/SandboxedEditing/Source/SandboxedEditing/Private/Framework/CommandMappings.cpp
new file mode 100644
--- /dev/null
+++ b/Engine/Plugins/Developer/Sandbox/SandboxedEditing/Source/SandboxedEditing/Private/Framework/CommandMappings.cpp
@@ -0,0 +1,70 @@
+﻿// Copyright Epic Games, Inc. All Rights Reserved.
+
+#include "CommandMappings.h"
+
+#include "Notifications.h"
+#include "SandboxedEditingCommands.h"
+#include "Framework/Commands/UICommandList.h"
+#include "ViewModels/SandboxSystemViewModel.h"
+#include "Widgets/Dialogues/SCreateNewSandbox.h"
+
+namespace UE::SandboxedEditing
+{
+FCommandMappings::FCommandMappings(const TSharedRef<FSandboxSystemViewModel>& InSandboxViewModel)
+	: SandboxViewModel(InSandboxViewModel)
+	, CommandList(MakeShared<FUICommandList>())
+{
+	FSandboxedEditingCommands::Register();
+	
+	FSandboxedEditingCommands& Commands = FSandboxedEditingCommands::Get();
+	
+	CommandList->MapAction(
+		Commands.OpenCreateNewSandboxDialog,
+		FExecuteAction::CreateStatic(&ShowNewSandboxModal, InSandboxViewModel.ToSharedPtr())
+		);
+	
+	CommandList->MapAction(
+		Commands.LeaveSandbox,
+		FExecuteAction::CreateRaw(this, &FCommandMappings::HandleLeaveSandbox),
+		FCanExecuteAction::CreateSP(InSandboxViewModel, &FSandboxSystemViewModel::CanLeaveSandbox)
+		);
+
+	CommandList->MapAction(
+		Commands.PersistAll,
+		FExecuteAction::CreateRaw(this, &FCommandMappings::HandlePersistAll),
+		FCanExecuteAction::CreateSP(InSandboxViewModel, &FSandboxSystemViewModel::CanPersistAllChanges)
+	);
+
+	CommandList->MapAction(
+		Commands.DiscardAll,
+		FExecuteAction::CreateRaw(this, &FCommandMappings::HandleDiscardAll),
+		FCanExecuteAction::CreateSP(InSandboxViewModel, &FSandboxSystemViewModel::CanDiscardAllChanges)
+	);
+}
+
+FCommandMappings::~FCommandMappings()
+{
+	FSandboxedEditingCommands::Unregister();
+}
+
+void FCommandMappings::HandleLeaveSandbox()
+{
+	const FString SandboxName = SandboxViewModel->GetActiveSandboxName();
+	SandboxViewModel->LeaveSandbox();
+
+	if (!SandboxName.IsEmpty()) // We expect this to be true
+	{
+		ShowLeftSandbox(*SandboxName);
+	}
+}
+
+void FCommandMappings::HandlePersistAll()
+{
+	SandboxViewModel->PersistAllChanges();
+}
+
+void FCommandMappings::HandleDiscardAll()
+{
+	SandboxViewModel->DiscardAllChanges();
+}
+}
diff --git a/Engine/Plugins/Developer/Sandbox/SandboxedEditing/Source/SandboxedEditing/Private/Framework/CommandMappings.h b/Engine/Plugins/Developer/Sandbox/SandboxedEditing/Source/SandboxedEditing/Private/Framework/CommandMappings.h
new file mode 100644
--- /dev/null
+++ b/Engine/Plugins/Developer/Sandbox/SandboxedEditing/Source/SandboxedEditing/Private/Framework/CommandMappings.h
@@ -0,0 +1,37 @@
+﻿// Copyright Epic Games, Inc. All Rights Reserved.
+
+#pragma once
+
+#include "Templates/SharedPointer.h"
+#include "Templates/UnrealTemplate.h"
+
+class FUICommandList;
+
+namespace UE::SandboxedEditing
+{
+class FSandboxSystemViewModel;
+	
+/** Handles all commands for our app. Has the FUICommandList and binds all commands. */
+class FCommandMappings : public FNoncopyable
+{
+public:
+	
+	explicit FCommandMappings(const TSharedRef<FSandboxSystemViewModel>& InSandboxViewModel);
+	~FCommandMappings();
+
+	const TSharedRef<FUICommandList>& GetCommandList() const { return CommandList; }
+
+private: 
+
+	/** Required to handle the commands. */
+	const TSharedRef<FSandboxSystemViewModel> SandboxViewModel;
+	
+	/** Holds our command bindings */
+	const TSharedRef<FUICommandList> CommandList;
+	
+	void HandleLeaveSandbox();
+	void HandlePersistAll();
+	void HandleDiscardAll();
+};
+}
+
diff --git a/Engine/Plugins/Developer/Sandbox/SandboxedEditing/Source/SandboxedEditing/Private/Framework/Data/CreateSandboxArgs.h b/Engine/Plugins/Developer/Sandbox/SandboxedEditing/Source/SandboxedEditing/Private/Framework/Data/CreateSandboxArgs.h
new file mode 100644
--- /dev/null
+++ b/Engine/Plugins/Developer/Sandbox/SandboxedEditing/Source/SandboxedEditing/Private/Framework/Data/CreateSandboxArgs.h
@@ -0,0 +1,14 @@
+﻿// Copyright Epic Games, Inc. All Rights Reserved.
+
+#pragma once
+
+#include "Containers/UnrealString.h"
+
+namespace UE::SandboxedEditing
+{
+struct FCreateSandboxArgs
+{
+	FString Name;
+	FString Description;
+};
+}
\ No newline at end of file
diff --git a/Engine/Plugins/Developer/Sandbox/SandboxedEditing/Source/SandboxedEditing/Private/Framework/Data/SandboxItem.h b/Engine/Plugins/Developer/Sandbox/SandboxedEditing/Source/SandboxedEditing/Private/Framework/Data/SandboxItem.h
new file mode 100644
--- /dev/null
+++ b/Engine/Plugins/Developer/Sandbox/SandboxedEditing/Source/SandboxedEditing/Private/Framework/Data/SandboxItem.h
@@ -0,0 +1,33 @@
+﻿// Copyright Epic Games, Inc. All Rights Reserved.
+
+#pragma once
+
+#include "Containers/UnrealString.h"
+
+namespace UE::SandboxedEditing
+{
+/** Describes a sandbox. */
+struct FSandboxItem
+{
+	/** Name the user has given the sandbox. May not be unique (but usually is). */
+	FString Name;
+
+	/** Description the user has given the sandbox. */
+	FString Description;
+
+	/** Path to the root of the sandbox directory. Unique for every sandbox. */
+	FString SandboxRoot;
+
+	FSandboxItem() = default;
+	explicit FSandboxItem(FString InName, FString InDescription, FString InDirectoryPath)
+		: Name(MoveTemp(InName))
+		, Description(MoveTemp(InDescription))
+		, SandboxRoot(MoveTemp(InDirectoryPath))
+	{}
+
+	friend bool operator<(const FSandboxItem& Left, const FSandboxItem& Right)
+	{
+		return Left.Name < Right.Name;
+	}
+};
+}
diff --git a/Engine/Plugins/Developer/Sandbox/SandboxedEditing/Source/SandboxedEditing/Private/Framework/Notifications.cpp b/Engine/Plugins/Developer/Sandbox/SandboxedEditing/Source/SandboxedEditing/Private/Framework/Notifications.cpp
new file mode 100644
--- /dev/null
+++ b/Engine/Plugins/Developer/Sandbox/SandboxedEditing/Source/SandboxedEditing/Private/Framework/Notifications.cpp
@@ -0,0 +1,84 @@
+﻿// Copyright Epic Games, Inc. All Rights Reserved.
+
+#include "Notifications.h"
+
+#include "Framework/Application/SlateApplication.h"
+#include "Framework/Notifications/NotificationManager.h"
+#include "Widgets/Notifications/SNotificationList.h"
+
+#define LOCTEXT_NAMESPACE "SandboxedEditing.Notifications"
+
+namespace UE::SandboxedEditing
+{
+void ShowCreatedSandbox(const FString& InSandboxName)
+{
+	if (!FSlateApplication::IsInitialized())
+	{
+		return;
+	}
+	
+	FSlateNotificationManager& NotificationManager = FSlateNotificationManager::Get();
+	FNotificationInfo Args(
+		FText::Format(LOCTEXT("JoinedFmt", "Joined sandbox {0}"), FText::FromString(InSandboxName))
+		);
+	Args.bFireAndForget = true;
+	Args.ExpireDuration = 2.f;
+	
+	NotificationManager.AddNotification(Args);
+}
+
+void ShowLoadedSandbox(const FString& InSandboxName)
+{
+	if (!FSlateApplication::IsInitialized())
+	{
+		return;
+	}
+	
+	FSlateNotificationManager& NotificationManager = FSlateNotificationManager::Get();
+	FNotificationInfo Args(
+		FText::Format(LOCTEXT("LoadedFmt", "Loaded sandbox {0}"), FText::FromString(InSandboxName))
+		);
+	Args.bFireAndForget = true;
+	Args.ExpireDuration = 2.f;
+	
+	NotificationManager.AddNotification(Args);
+}
+
+void ShowFailedToLoadSandbox(const FString& InSandboxName)
+{
+	if (!FSlateApplication::IsInitialized())
+	{
+		return;
+	}
+	
+	FSlateNotificationManager& NotificationManager = FSlateNotificationManager::Get();
+	FNotificationInfo Args(
+		FText::Format(LOCTEXT("FailedToLoadFmt.Title", "Failed to load sandbox {0}"), FText::FromString(InSandboxName))
+		);
+	Args.bFireAndForget = true;
+	Args.ExpireDuration = 2.f;
+	
+	const TSharedPtr<SNotificationItem> Item = NotificationManager.AddNotification(Args);
+	Item->SetCompletionState(SNotificationItem::CS_Fail);
+	Item->SetSubText(LOCTEXT("FailedToLoadFmt.Description", "See log for more info."));
+}
+
+void ShowLeftSandbox(const FString& InSandboxName)
+{
+	if (!FSlateApplication::IsInitialized())
+	{
+		return;
+	}
+	
+	FSlateNotificationManager& NotificationManager = FSlateNotificationManager::Get();
+	FNotificationInfo Args(
+		FText::Format(LOCTEXT("LeftFmt", "Left sandbox {0}"), FText::FromString(InSandboxName))
+		);
+	Args.bFireAndForget = true;
+	Args.ExpireDuration = 2.f;
+	
+	NotificationManager.AddNotification(Args);
+}
+}
+
+#undef LOCTEXT_NAMESPACE
\ No newline at end of file
diff --git a/Engine/Plugins/Developer/Sandbox/SandboxedEditing/Source/SandboxedEditing/Private/Framework/Notifications.h b/Engine/Plugins/Developer/Sandbox/SandboxedEditing/Source/SandboxedEditing/Private/Framework/Notifications.h
new file mode 100644
--- /dev/null
+++ b/Engine/Plugins/Developer/Sandbox/SandboxedEditing/Source/SandboxedEditing/Private/Framework/Notifications.h
@@ -0,0 +1,19 @@
+﻿// Copyright Epic Games, Inc. All Rights Reserved.
+
+#pragma once
+
+class FString;
+
+namespace UE::SandboxedEditing
+{
+/** Shows a notification that the user joined a sandbox. */
+void ShowCreatedSandbox(const FString& InSandboxName);
+	
+/** Shows the sandbox was loaded successfully. */
+void ShowLoadedSandbox(const FString& InSandboxName);
+/** Show that the sandbox could not be loaded. */
+void ShowFailedToLoadSandbox(const FString& InSandboxName);
+
+/** Shows notification that the user left a sandbox. */
+void ShowLeftSandbox(const FString& InSandboxName);
+}
diff --git a/Engine/Plugins/Developer/Sandbox/SandboxedEditing/Source/SandboxedEditing/Private/Framework/SandboxedEditingApp.cpp b/Engine/Plugins/Developer/Sandbox/SandboxedEditing/Source/SandboxedEditing/Private/Framework/SandboxedEditingApp.cpp
new file mode 100644
--- /dev/null
+++ b/Engine/Plugins/Developer/Sandbox/SandboxedEditing/Source/SandboxedEditing/Private/Framework/SandboxedEditingApp.cpp
@@ -0,0 +1,40 @@
+﻿// Copyright Epic Games, Inc. All Rights Reserved.
+
+#include "SandboxedEditingApp.h"
+
+#include "Internationalization/Text.h"
+#include "ToolMenus.h"
+#include "Widgets/SSandboxStatusBarWidget.h"
+
+namespace UE::SandboxedEditing
+{
+namespace Private
+{
+static TSharedRef<SWidget> CreateStatusBarWidget(const TSharedRef<FSandboxSystemViewModel>& InViewModel, const FCommandMappings& InCommandMappings)
+{
+	return SNew(SSandboxStatusBarWidget, InViewModel, InCommandMappings.GetCommandList());
+}
+	
+static void RegisterStatusBar(const TSharedRef<FSandboxSystemViewModel>& InViewModel, const FCommandMappings& InCommandMappings)
+{
+	UToolMenu* Menu = UToolMenus::Get()->ExtendMenu(TEXT("LevelEditor.StatusBar.ToolBar"));
+	FToolMenuSection& AssetManagementSection = Menu->AddSection(
+		TEXT("Asset Management"), FText::GetEmpty(), FToolMenuInsert(NAME_None, EToolMenuInsertType::First)
+		);
+
+	AssetManagementSection.AddEntry(
+		FToolMenuEntry::InitWidget(TEXT("SandboxStatusBar"), CreateStatusBarWidget(InViewModel, InCommandMappings), FText::GetEmpty(), true, false)
+	);
+}
+}
+
+FSandboxedEditingApp::FSandboxedEditingApp()
+	: SandboxViewModel(MakeShared<FSandboxSystemViewModel>())
+	, CommandMappings(SandboxViewModel)
+{
+	FCoreDelegates::OnPostEngineInit.AddLambda([this]
+	{
+		Private::RegisterStatusBar(SandboxViewModel, CommandMappings);
+	});
+}
+}
diff --git a/Engine/Plugins/Developer/Sandbox/SandboxedEditing/Source/SandboxedEditing/Private/Framework/SandboxedEditingApp.h b/Engine/Plugins/Developer/Sandbox/SandboxedEditing/Source/SandboxedEditing/Private/Framework/SandboxedEditingApp.h
new file mode 100644
--- /dev/null
+++ b/Engine/Plugins/Developer/Sandbox/SandboxedEditing/Source/SandboxedEditing/Private/Framework/SandboxedEditingApp.h
@@ -0,0 +1,34 @@
+﻿// Copyright Epic Games, Inc. All Rights Reserved.
+
+#pragma once
+
+#include "CommandMappings.h"
+#include "Templates/SharedPointer.h"
+#include "Templates/UnrealTemplate.h"
+#include "ViewModels/SandboxSystemViewModel.h"
+
+namespace UE::SandboxedEditing
+{
+/**
+ * Creates all relevant systems to deliver the sandboxed editing experience in the editor.
+ * This is a coordinator: It sets up editor hooks. It manages the creation and binding of view-models and views.
+ */
+class FSandboxedEditingApp : public FNoncopyable
+{
+public:
+
+	FSandboxedEditingApp();
+
+private:
+
+	/**
+	 * View model for interacting with sandbox system.
+	 * Shared pointer in case the UI to avoid crash in case that though some bug the UI outlives this app for whatever reason.
+	 */
+	const TSharedRef<FSandboxSystemViewModel> SandboxViewModel;
+
+	/** Handles all the input commands to this app. */
+	const FCommandMappings CommandMappings;
+};
+}
+
diff --git a/Engine/Plugins/Developer/Sandbox/SandboxedEditing/Source/SandboxedEditing/Private/Framework/ViewModels/SandboxSystemViewModel.cpp b/Engine/Plugins/Developer/Sandbox/SandboxedEditing/Source/SandboxedEditing/Private/Framework/ViewModels/SandboxSystemViewModel.cpp
new file mode 100644
--- /dev/null
+++ b/Engine/Plugins/Developer/Sandbox/SandboxedEditing/Source/SandboxedEditing/Private/Framework/ViewModels/SandboxSystemViewModel.cpp
@@ -0,0 +1,169 @@
+﻿// Copyright Epic Games, Inc. All Rights Reserved.
+
+#include "SandboxSystemViewModel.h"
+
+#include "Data/SandboxMetaData.h"
+#include "Framework/Data/CreateSandboxArgs.h"
+#include "Framework/Data/SandboxItem.h"
+#include "IFileSandboxCoreModule.h"
+#include "ISandboxInstance.h"
+#include "ISandboxManager.h"
+#include "Misc/Paths.h"
+#include "Templates/Function.h"
+#include "Types/EBreakBehavior.h"
+#include "Types/Manager/LoadSandboxByDirectoryArgs.h"
+#include "Types/Manager/NewSandboxArgs.h"
+
+#define LOCTEXT_NAMESPACE "FSandboxSystemViewModel"
+
+namespace UE::SandboxedEditing
+{
+bool FSandboxSystemViewModel::CreateNewSandbox(const FCreateSandboxArgs& InArgs)
+{
+	if (!CanCreateNewSandbox(InArgs.Name))
+	{
+		return false;
+	}
+	
+	using namespace FileSandboxCore;
+	IFileSandboxCoreModule& Module = IFileSandboxCoreModule::Get();
+	ISandboxManager& SandboxManager = Module.GetSandboxManager();
+
+	if (SandboxManager.HasActiveSandbox())
+	{
+		SandboxManager.LeaveSandbox();
+	}
+
+	const FNewSandboxArgs Args(InArgs.Name, InArgs.Description);
+	return SandboxManager.CreateNewSandbox(Args).HasValue();
+}
+
+bool FSandboxSystemViewModel::CanCreateNewSandbox(const FString& InName, FText* Reason)
+{
+	FileSandboxCore::EDirectorySuitability Suitability = FileSandboxCore::DetermineDirectorySuitability(InName);
+
+	if (Reason)
+	{
+		*Reason = FileSandboxCore::FormatDirectorySuitabilityAsReason(Suitability, InName);
+	}
+	
+	return Suitability == FileSandboxCore::EDirectorySuitability::Suitable;
+}
+
+bool FSandboxSystemViewModel::LoadSandbox(const FString& InRootDirectory)
+{
+	if (CanLoadSandbox(InRootDirectory))
+	{
+		using namespace FileSandboxCore;
+		IFileSandboxCoreModule& Module = IFileSandboxCoreModule::Get();
+		ISandboxManager& SandboxManager = Module.GetSandboxManager();
+		return SandboxManager.LoadSandbox(FLoadSandboxByDirectoryArgs{ InRootDirectory }).HasValue();
+	}
+	return false;
+}
+
+bool FSandboxSystemViewModel::CanLoadSandbox(const FString& InRootDirectory)
+{
+	return !IsActiveSandbox(InRootDirectory);
+}
+
+bool FSandboxSystemViewModel::IsActiveSandbox(const FString& InRootDirectory) const
+{
+	using namespace FileSandboxCore;
+	IFileSandboxCoreModule& Module = IFileSandboxCoreModule::Get();
+	ISandboxManager& SandboxManager = Module.GetSandboxManager();
+	ISandboxInstance* Instance = SandboxManager.GetActiveSandboxInstance();
+	return Instance && FPaths::IsSamePath(Instance->GetRootDirectory(), InRootDirectory);
+}
+
+void FSandboxSystemViewModel::LeaveSandbox() const
+{
+	using namespace FileSandboxCore;
+	IFileSandboxCoreModule& Module = IFileSandboxCoreModule::Get();
+	ISandboxManager& SandboxManager = Module.GetSandboxManager();
+	SandboxManager.LeaveSandbox();
+}
+
+bool FSandboxSystemViewModel::CanLeaveSandbox() const
+{
+	using namespace FileSandboxCore;
+	IFileSandboxCoreModule& Module = IFileSandboxCoreModule::Get();
+	ISandboxManager& SandboxManager = Module.GetSandboxManager();
+	return SandboxManager.GetActiveSandboxInstance() != nullptr;
+}
+
+void FSandboxSystemViewModel::PersistAllChanges() const
+{
+	using namespace FileSandboxCore;
+	IFileSandboxCoreModule& Module = IFileSandboxCoreModule::Get();
+	ISandboxManager& SandboxManager = Module.GetSandboxManager();
+	ISandboxInstance* Instance = SandboxManager.GetActiveSandboxInstance();
+	
+	if (CanPersistAllChanges() && ensure(Instance))
+	{
+		Instance->PersistAll();
+	}
+}
+
+bool FSandboxSystemViewModel::CanPersistAllChanges() const
+{
+	return HasActiveSandbox();
+}
+
+void FSandboxSystemViewModel::DiscardAllChanges() const
+{
+	using namespace FileSandboxCore;
+	IFileSandboxCoreModule& Module = IFileSandboxCoreModule::Get();
+	ISandboxManager& SandboxManager = Module.GetSandboxManager();
+	ISandboxInstance* Instance = SandboxManager.GetActiveSandboxInstance();
+	
+	if (CanDiscardAllChanges() && ensure(Instance))
+	{
+		Instance->DiscardAll();
+	}
+}
+
+bool FSandboxSystemViewModel::CanDiscardAllChanges() const
+{
+	return HasActiveSandbox();
+}
+
+FString FSandboxSystemViewModel::GetActiveSandboxName() const
+{
+	using namespace FileSandboxCore;
+	IFileSandboxCoreModule& Module = IFileSandboxCoreModule::Get();
+	ISandboxInstance* SandboxInstance = Module.GetSandboxManager().GetActiveSandboxInstance();
+	return SandboxInstance ? SandboxInstance->GetInitialMetaData().Name : FString();
+}
+
+bool FSandboxSystemViewModel::HasActiveSandbox() const
+{
+	using namespace FileSandboxCore;
+	IFileSandboxCoreModule& Module = IFileSandboxCoreModule::Get();
+	return Module.GetSandboxManager().GetActiveSandboxInstance() != nullptr;
+}
+
+TArray<FSandboxItem> FSandboxSystemViewModel::GetKnownSandboxes() const
+{
+	using namespace FileSandboxCore;
+	
+	// This is very inefficient because it loads up every single manifest file, parses the JSON, etc.
+	// We should add a cache to FileSandboxCore for this.
+	TArray<FSandboxItem> SandboxDirectories;
+	ForEachSandbox([&SandboxDirectories](const FString& Directory)
+	{
+		if (TOptional<FFileSandboxCore_SandboxMetaData> MetaData = LoadMetaData(Directory))
+		{
+			SandboxDirectories.Emplace(MoveTemp(MetaData->Name), MoveTemp(MetaData->Description), Directory);
+		}
+		return EBreakBehavior::Continue;
+	});
+
+	// Display them in alphabetical order
+	SandboxDirectories.Sort();
+	
+	return SandboxDirectories;
+}
+}
+
+#undef LOCTEXT_NAMESPACE
\ No newline at end of file
diff --git a/Engine/Plugins/Developer/Sandbox/SandboxedEditing/Source/SandboxedEditing/Private/Framework/ViewModels/SandboxSystemViewModel.h b/Engine/Plugins/Developer/Sandbox/SandboxedEditing/Source/SandboxedEditing/Private/Framework/ViewModels/SandboxSystemViewModel.h
new file mode 100644
--- /dev/null
+++ b/Engine/Plugins/Developer/Sandbox/SandboxedEditing/Source/SandboxedEditing/Private/Framework/ViewModels/SandboxSystemViewModel.h
@@ -0,0 +1,62 @@
+﻿// Copyright Epic Games, Inc. All Rights Reserved.
+
+#pragma once
+
+#include "Delegates/Delegate.h"
+#include "Templates/UnrealTemplate.h"
+
+template<typename OptionalType> struct TOptional;
+
+namespace UE::SandboxedEditing
+{
+struct FCreateSandboxArgs;
+struct FSandboxItem;
+	
+/**
+ * View model for creating, loading, and querying the active sandbox.
+ * 
+ * Bridges the sandbox system without exposing it directly.
+ * Members added here should have direct corresponding UI elements.
+ */
+class FSandboxSystemViewModel : public FNoncopyable
+{
+public:
+
+	/** Creates a new sandbox with the given name. */
+	bool CreateNewSandbox(const FCreateSandboxArgs& InArgs);
+	/** @return Whether it is valid to call CreateNewSandbox. */
+	bool CanCreateNewSandbox(const FString& InName, FText* Reason = nullptr);
+
+	/** Loads the specified sandbox. */
+	bool LoadSandbox(const FString& InRootDirectory);
+	/** @return Whether the user can join this sandbox (i.e. not yet in this sandbox) */
+	bool CanLoadSandbox(const FString& InRootDirectory);
+	/** @return Whether this is the active sandbox. */
+	bool IsActiveSandbox(const FString& InRootDirectory) const;
+
+	/** Leaves the current sandbox. */
+	void LeaveSandbox() const;
+	/** @return Whether the sandbox can be left (false if there is no active one). */
+	bool CanLeaveSandbox() const;
+
+	/** Persists all changes */
+	void PersistAllChanges() const;
+	/** @return Whether all changes can be persisted */
+	bool CanPersistAllChanges() const;
+
+	/** Discards all changes */
+	void DiscardAllChanges() const;
+	/** @return Whether all changes can be discarded */
+	bool CanDiscardAllChanges() const; 
+
+	/** @return Name of the currently active sandbox. Empty if none active. */
+	FString GetActiveSandboxName() const;
+
+	/** @return Whether there is an active sandbox */
+	bool HasActiveSandbox() const;
+
+	/** @return The sandboxes we know of. */
+	TArray<FSandboxItem> GetKnownSandboxes() const;
+};
+}
+
diff --git a/Engine/Plugins/Developer/Sandbox/SandboxedEditing/Source/SandboxedEditing/Private/SandboxEditingModule.cpp b/Engine/Plugins/Developer/Sandbox/SandboxedEditing/Source/SandboxedEditing/Private/SandboxEditingModule.cpp
new file mode 100644
--- /dev/null
+++ b/Engine/Plugins/Developer/Sandbox/SandboxedEditing/Source/SandboxedEditing/Private/SandboxEditingModule.cpp
@@ -0,0 +1,21 @@
+// Copyright Epic Games, Inc. All Rights Reserved.
+
+#include "SandboxEditingModule.h"
+
+#include "SandboxedEditingCommands.h"
+#include "SandboxedEditingStyle.h"
+
+namespace UE::SandboxedEditing
+{
+void FSandboxEditingModule::StartupModule()
+{
+	FSandboxedEditingStyle::Get();
+}
+
+void FSandboxEditingModule::ShutdownModule()
+{
+	
+}
+}
+
+IMPLEMENT_MODULE(UE::SandboxedEditing::FSandboxEditingModule, SandboxedEditing);
\ No newline at end of file
diff --git a/Engine/Plugins/Developer/Sandbox/SandboxedEditing/Source/SandboxedEditing/Private/SandboxEditingModule.h b/Engine/Plugins/Developer/Sandbox/SandboxedEditing/Source/SandboxedEditing/Private/SandboxEditingModule.h
new file mode 100644
--- /dev/null
+++ b/Engine/Plugins/Developer/Sandbox/SandboxedEditing/Source/SandboxedEditing/Private/SandboxEditingModule.h
@@ -0,0 +1,25 @@
+// Copyright Epic Games, Inc. All Rights Reserved.
+
+#pragma once
+
+#include "ISandboxEditingModule.h"
+#include "Framework/SandboxedEditingApp.h"
+
+namespace UE::SandboxedEditing
+{
+class FSandboxEditingModule : public ISandboxedEditingModule
+{
+public:
+
+	//~ Begin IModuleInterface Interface
+	virtual void StartupModule() override;
+	virtual void ShutdownModule() override;
+	//~ End IModuleInterface Interface
+
+private:
+
+	/** Manages all the business logic and creates the editor workflows. */
+	FSandboxedEditingApp EditingApp;
+};
+}
+
diff --git a/Engine/Plugins/Developer/Sandbox/SandboxedEditing/Source/SandboxedEditing/Private/SandboxedEditingCommands.cpp b/Engine/Plugins/Developer/Sandbox/SandboxedEditing/Source/SandboxedEditing/Private/SandboxedEditingCommands.cpp
new file mode 100644
--- /dev/null
+++ b/Engine/Plugins/Developer/Sandbox/SandboxedEditing/Source/SandboxedEditing/Private/SandboxedEditingCommands.cpp
@@ -0,0 +1,38 @@
+﻿// Copyright Epic Games, Inc. All Rights Reserved.
+
+#include "SandboxedEditingCommands.h"
+
+#include "Styling/AppStyle.h"
+
+#define LOCTEXT_NAMESPACE "FSandboxedEditingCommands"
+
+namespace UE::SandboxedEditing
+{
+FSandboxedEditingCommands::FSandboxedEditingCommands()
+	: TCommands<FSandboxedEditingCommands>
+	(
+		"SandboxedEditing",
+		NSLOCTEXT("Contexts", "SandboxedEditing", "SandboxedEditing"),
+		"LevelEditor",
+		FAppStyle::GetAppStyleSetName()
+	)
+{}
+
+void FSandboxedEditingCommands::RegisterCommands()
+{
+	UI_COMMAND(OpenCreateNewSandboxDialog,
+		"New Sandbox", "Creates a new sandbox for asset modifications.", EUserInterfaceActionType::Button, FInputChord()
+		);
+	UI_COMMAND(LeaveSandbox,
+		"Leave Sandbox", "Leaves the active sandbox without persisting.", EUserInterfaceActionType::Button, FInputChord()
+		);
+	UI_COMMAND(PersistAll,
+		"Persist All", "Persists all changes.", EUserInterfaceActionType::Button, FInputChord()
+		);
+	UI_COMMAND(DiscardAll,
+		"Discard All", "Discards all changes.", EUserInterfaceActionType::Button, FInputChord()
+		);
+}
+}
+
+#undef LOCTEXT_NAMESPACE
\ No newline at end of file
diff --git a/Engine/Plugins/Developer/Sandbox/SandboxedEditing/Source/SandboxedEditing/Private/SandboxedEditingCommands.h b/Engine/Plugins/Developer/Sandbox/SandboxedEditing/Source/SandboxedEditing/Private/SandboxedEditingCommands.h
new file mode 100644
--- /dev/null
+++ b/Engine/Plugins/Developer/Sandbox/SandboxedEditing/Source/SandboxedEditing/Private/SandboxedEditingCommands.h
@@ -0,0 +1,34 @@
+﻿// Copyright Epic Games, Inc. All Rights Reserved.
+
+#pragma once
+
+#include "Framework/Commands/Commands.h"
+
+namespace UE::SandboxedEditing
+{
+class FSandboxedEditingCommands : public TCommands<FSandboxedEditingCommands>
+{
+public:
+
+	FSandboxedEditingCommands();
+
+	//~ Begin TCommands Interface
+	virtual void RegisterCommands() override;
+	//~ End TCommands Interface
+
+	/**
+	 * Opens a dialog for creating a new sandbox.
+	 * @note This is temporary UX - this will be removed when we implement the final UX.
+	 */
+	TSharedPtr<FUICommandInfo> OpenCreateNewSandboxDialog;
+	
+	/** Leaves the active sandbox without persisting. */
+	TSharedPtr<FUICommandInfo> LeaveSandbox;
+	
+	/** Persists all changes. */
+	TSharedPtr<FUICommandInfo> PersistAll;
+	/** Discards all changes. */
+	TSharedPtr<FUICommandInfo> DiscardAll;
+};
+}
+
diff --git a/Engine/Plugins/Developer/Sandbox/SandboxedEditing/Source/SandboxedEditing/Private/SandboxedEditingStyle.cpp b/Engine/Plugins/Developer/Sandbox/SandboxedEditing/Source/SandboxedEditing/Private/SandboxedEditingStyle.cpp
new file mode 100644
--- /dev/null
+++ b/Engine/Plugins/Developer/Sandbox/SandboxedEditing/Source/SandboxedEditing/Private/SandboxedEditingStyle.cpp
@@ -0,0 +1,32 @@
+﻿// Copyright Epic Games, Inc. All Rights Reserved.
+
+#include "SandboxedEditingStyle.h"
+
+#include "Styling/SlateStyleRegistry.h"
+
+namespace UE::SandboxedEditing
+{
+FSandboxedEditingStyle& FSandboxedEditingStyle::Get()
+{
+	static FSandboxedEditingStyle Style;
+	return Style;
+}
+	
+FSandboxedEditingStyle::FSandboxedEditingStyle()
+	: FSlateStyleSet("SandboxedEditingStyle")
+{
+	const FString PluginContentDir = FPaths::EnginePluginsDir() / TEXT("Developer") / TEXT("Sandbox") / TEXT("SandboxedEditing") / TEXT("Resources");
+	const FString EngineEditorSlateDir = FPaths::EngineContentDir() / TEXT("Editor") / TEXT("Slate");
+	SetContentRoot(PluginContentDir);
+	SetCoreContentRoot(EngineEditorSlateDir);
+	
+	Set("SandboxedEditing.Icons.CreateNewSandbox", new FSlateColorBrush(FLinearColor::Red));
+	
+	FSlateStyleRegistry::RegisterSlateStyle(*this);
+}
+
+FSandboxedEditingStyle::~FSandboxedEditingStyle()
+{
+	FSlateStyleRegistry::UnRegisterSlateStyle(*this);
+}
+}
diff --git a/Engine/Plugins/Developer/Sandbox/SandboxedEditing/Source/SandboxedEditing/Private/SandboxedEditingStyle.h b/Engine/Plugins/Developer/Sandbox/SandboxedEditing/Source/SandboxedEditing/Private/SandboxedEditingStyle.h
new file mode 100644
--- /dev/null
+++ b/Engine/Plugins/Developer/Sandbox/SandboxedEditing/Source/SandboxedEditing/Private/SandboxedEditingStyle.h
@@ -0,0 +1,23 @@
+﻿// Copyright Epic Games, Inc. All Rights Reserved.
+
+#pragma once
+
+#include "Styling/ISlateStyle.h"
+#include "Styling/SlateStyle.h"
+
+class FSlateStyleSet;
+class ISlateStyle;
+
+namespace UE::SandboxedEditing
+{
+class FSandboxedEditingStyle final : public FSlateStyleSet
+{
+public:
+	
+	static FSandboxedEditingStyle& Get();
+
+	FSandboxedEditingStyle();
+	virtual ~FSandboxedEditingStyle() override;
+};
+}
+
diff --git a/Engine/Plugins/Developer/Sandbox/SandboxedEditing/Source/SandboxedEditing/Private/Widgets/Dialogues/SCreateNewSandbox.cpp b/Engine/Plugins/Developer/Sandbox/SandboxedEditing/Source/SandboxedEditing/Private/Widgets/Dialogues/SCreateNewSandbox.cpp
new file mode 100644
--- /dev/null
+++ b/Engine/Plugins/Developer/Sandbox/SandboxedEditing/Source/SandboxedEditing/Private/Widgets/Dialogues/SCreateNewSandbox.cpp
@@ -0,0 +1,217 @@
+﻿// Copyright Epic Games, Inc. All Rights Reserved.
+
+#include "SCreateNewSandbox.h"
+
+#include "Framework/Application/SlateApplication.h"
+#include "Framework/Data/CreateSandboxArgs.h"
+#include "Framework/ViewModels/SandboxSystemViewModel.h"
+#include "SPositiveActionButton.h"
+#include "Framework/Notifications.h"
+#include "Styling/AppStyle.h"
+#include "Widgets/DeclarativeSyntaxSupport.h"
+#include "Widgets/Docking/SDockTab.h"
+#include "Widgets/Input/SButton.h"
+#include "Widgets/Input/SEditableTextBox.h"
+#include "Widgets/Input/SMultiLineEditableTextBox.h"
+#include "Widgets/Layout/SBox.h"
+#include "Widgets/Layout/SExpandableArea.h"
+#include "Widgets/Layout/SUniformGridPanel.h"
+#include "Widgets/Notifications/SErrorText.h"
+#include "Widgets/SBoxPanel.h"
+#include "Widgets/SWindow.h"
+#include "Widgets/Text/STextBlock.h"
+
+#define LOCTEXT_NAMESPACE "SCreateNewSandbox"
+
+namespace UE::SandboxedEditing
+{
+void ShowNewSandboxModal(TSharedPtr<FSandboxSystemViewModel> InViewModel)
+{
+	if (!ensure(InViewModel.IsValid()))
+	{
+		return;
+	}
+	
+	const TSharedRef<SWindow> NewWindow = SNew(SWindow)
+		.Title(LOCTEXT("CreateSandbox", "Create new sandbox"))
+		.SizingRule(ESizingRule::Autosized)
+		.ClientSize(FVector2D(600, 600))
+		.SupportsMaximize(true)
+		.SupportsMinimize(false);
+
+	const TSharedRef<SCreateNewSandbox> CreateSandbox = SNew(SCreateNewSandbox, InViewModel.ToSharedRef()).ParentWindow(NewWindow);
+	NewWindow->SetContent(
+		CreateSandbox
+	);
+
+	FSlateApplication::Get().AddModalWindow(NewWindow, nullptr);
+	if (!CreateSandbox->IsDialogConfirmed())
+	{
+		return;
+	}
+
+	const FCreateSandboxArgs Args = CreateSandbox->GetArgs();
+	if (InViewModel->CreateNewSandbox(Args))
+	{
+		ShowCreatedSandbox(Args.Name);
+	}
+}
+	
+void SCreateNewSandbox::Construct(const FArguments& InArgs, const TSharedRef<FSandboxSystemViewModel>& InModel)
+{
+	ParentWindow = InArgs._ParentWindow;
+	Model = InModel;
+	
+	ChildSlot
+	[
+		SNew(SBorder)
+		.BorderImage(FAppStyle::GetBrush("ToolPanel.GroupBorder"))
+		[
+			SNew(SVerticalBox)
+			
+			+SVerticalBox::Slot()
+			.AutoHeight()
+			.Padding(2.)
+			[			
+				SAssignNew(NameTextCtrl, SEditableTextBox)
+				.HintText(LOCTEXT("Name.Hint", "Name"))
+				.OnTextChanged(this, &SCreateNewSandbox::OnNameChanged)
+			]
+			
+			+SVerticalBox::Slot()
+			.AutoHeight()
+			.Padding(2.)
+			[
+				CreateDescriptionInput()
+			]
+			
+			+SVerticalBox::Slot()
+			.AutoHeight()
+			.HAlign(HAlign_Right)
+			.VAlign(VAlign_Bottom)
+			.Padding(0.0f,0.0f,0.0f,5.0f)
+			[
+				CreateButtonArea()
+			]
+		]
+	];
+}
+
+TSharedRef<SWidget> SCreateNewSandbox::CreateDescriptionInput()
+{
+	return SNew(SExpandableArea)
+		.HeaderContent()
+		[
+			SNew(STextBlock)
+			.Text(LOCTEXT("Description.Label", "Description"))
+		]
+		.BodyContent()
+		[
+			SNew(SBox)
+			.MinDesiredHeight(120)
+			.WidthOverride(320)
+			[
+				SAssignNew(DescriptionTextCtrl, SMultiLineEditableTextBox)
+				.SelectAllTextWhenFocused(true)
+				.AutoWrapText(true)
+				.HintText(LOCTEXT("Description.Hint", "Optional description"))
+			]
+		];
+}
+	
+TSharedRef<SWidget> SCreateNewSandbox::CreateButtonArea()
+{
+	return SNew(SUniformGridPanel)
+		.SlotPadding(FAppStyle::GetMargin("StandardDialog.SlotPadding"))
+		.MinDesiredSlotWidth(FAppStyle::GetFloat("StandardDialog.MinDesiredSlotWidth"))
+		.MinDesiredSlotHeight(FAppStyle::GetFloat("StandardDialog.MinDesiredSlotHeight"))
+		+ SUniformGridPanel::Slot(0, 0)
+		[
+			SNew(SPositiveActionButton)
+			.IsEnabled(this, &SCreateNewSandbox::IsCreateButtonEnabled)
+			.Text(LOCTEXT("ButtonCreate", "Create"))
+			.ToolTipText(this, &SCreateNewSandbox::GetCreateButtonToolTipText)
+			.OnClicked(this, &SCreateNewSandbox::OnCreateClicked)
+		]
+		+ SUniformGridPanel::Slot(1, 0)
+		[
+			SNew(SActionButton)
+			.Text(LOCTEXT("CancelButton", "Cancel"))
+			.OnClicked(this, &SCreateNewSandbox::OnCancelClicked)
+		];
+}
+
+FReply SCreateNewSandbox::OnKeyDown(const FGeometry& MyGeometry, const FKeyEvent& InKeyEvent)
+{
+	// Pressing escape returns as if the user clicked cancel
+	if (InKeyEvent.GetKey() == EKeys::Escape)
+	{
+		return OnCancelClicked();
+	}
+
+	return FReply::Unhandled();
+}
+
+FCreateSandboxArgs SCreateNewSandbox::GetArgs() const
+{
+	FCreateSandboxArgs Args;
+	Args.Name = NameTextCtrl->GetText().ToString();
+	Args.Description = DescriptionTextCtrl->GetText().ToString();
+	return Args;
+}
+
+FReply SCreateNewSandbox::OnCreateClicked()
+{
+	bDialogConfirmed = true;
+	auto ParentWindowPin = ParentWindow.Pin();
+	if (ParentWindowPin.IsValid())
+	{
+		ParentWindowPin->RequestDestroyWindow();
+	}
+
+	return FReply::Handled();
+}
+
+bool SCreateNewSandbox::IsCreateButtonEnabled() const
+{
+	return Model->CanCreateNewSandbox(NameTextCtrl->GetText().ToString());
+}
+
+FText SCreateNewSandbox::GetCreateButtonToolTipText() const
+{
+	FText Reason;
+	const bool bCanCreateSandbox = Model->CanCreateNewSandbox(NameTextCtrl->GetText().ToString(), &Reason);
+	return bCanCreateSandbox
+		? LOCTEXT("CanCreateToolTip", "Create sandbox")
+		: Reason;
+}
+
+void SCreateNewSandbox::OnNameChanged(const FText& InNewName)
+{
+	FText Reason;
+	const bool bCanCreateSandbox = Model->CanCreateNewSandbox(InNewName.ToString(), &Reason);
+	
+	if (bCanCreateSandbox) 
+	{
+		NameTextCtrl->SetError(FText::GetEmpty());
+	}
+	else
+	{
+		NameTextCtrl->SetError(Reason);
+	}
+}
+
+FReply SCreateNewSandbox::OnCancelClicked()
+{
+	bDialogConfirmed = false;
+	auto ParentWindowPin = ParentWindow.Pin();
+	if (ParentWindowPin.IsValid())
+	{
+		ParentWindowPin->RequestDestroyWindow();
+	}
+
+	return FReply::Handled();
+}
+}
+
+#undef LOCTEXT_NAMESPACE
\ No newline at end of file
diff --git a/Engine/Plugins/Developer/Sandbox/SandboxedEditing/Source/SandboxedEditing/Private/Widgets/Dialogues/SCreateNewSandbox.h b/Engine/Plugins/Developer/Sandbox/SandboxedEditing/Source/SandboxedEditing/Private/Widgets/Dialogues/SCreateNewSandbox.h
new file mode 100644
--- /dev/null
+++ b/Engine/Plugins/Developer/Sandbox/SandboxedEditing/Source/SandboxedEditing/Private/Widgets/Dialogues/SCreateNewSandbox.h
@@ -0,0 +1,73 @@
+﻿// Copyright Epic Games, Inc. All Rights Reserved.
+
+#pragma once
+
+#include "Widgets/SCompoundWidget.h"
+#include "Widgets/DeclarativeSyntaxSupport.h"
+
+class SEditableTextBox;
+class SMultiLineEditableTextBox;
+class SErrorText;
+
+namespace UE::SandboxedEditing
+{
+class FSandboxSystemViewModel;
+struct FCreateSandboxArgs;
+	
+/** Shows a modal dialogue to the user for starting a new sandbox. */
+void ShowNewSandboxModal(TSharedPtr<FSandboxSystemViewModel> InViewModel);
+	
+class SCreateNewSandbox : public SCompoundWidget
+{
+public:
+
+	SLATE_BEGIN_ARGS(SCreateNewSandbox){}
+		/** The parent window this widget is hosted in. */
+		SLATE_ARGUMENT(TSharedPtr<SWindow>, ParentWindow)
+	SLATE_END_ARGS()
+
+	void Construct(const FArguments& InArgs, const TSharedRef<FSandboxSystemViewModel>& InModel);
+
+	virtual FReply OnKeyDown(const FGeometry& MyGeometry, const FKeyEvent& InKeyEvent) override;
+
+	/** Get dialog result */
+	bool IsDialogConfirmed() const { return bDialogConfirmed; }
+
+	/** Get the new sandbox command from the dialog result. */
+	FCreateSandboxArgs GetArgs() const;
+		
+private:
+	
+	/** Pointer to the parent modal window */
+	TWeakPtr<SWindow> ParentWindow;
+
+	/** Model use to determine whether a name is available. */
+	TSharedPtr<FSandboxSystemViewModel> Model;
+	
+	TSharedPtr<SMultiLineEditableTextBox> DescriptionTextCtrl;
+	TSharedPtr<SEditableTextBox> NameTextCtrl;
+	
+	/** Result confirmation */
+	bool bDialogConfirmed = false;
+	
+	TSharedRef<SWidget> CreateDescriptionInput();
+	TSharedRef<SWidget> CreateButtonArea();
+	
+	/** Called when the settings of the dialog are to be accepted*/
+	FReply OnCreateClicked();
+	/** Called to check if create button is enabled or not. */
+	bool IsCreateButtonEnabled() const;
+	/** @return Tooltip text to show on create button */
+	FText GetCreateButtonToolTipText() const;
+
+	/** @return Whether the sandbox can be created. */
+	bool CanCreate(FText* OutReason);
+
+	/** Updates the error state when the name changes. */
+	void OnNameChanged(const FText& InNewName);
+	
+	/** Called when the sandbox creation is abandoned. */
+	FReply OnCancelClicked();
+};
+}
+
diff --git a/Engine/Plugins/Developer/Sandbox/SandboxedEditing/Source/SandboxedEditing/Private/Widgets/SSandboxStatusBarWidget.cpp b/Engine/Plugins/Developer/Sandbox/SandboxedEditing/Source/SandboxedEditing/Private/Widgets/SSandboxStatusBarWidget.cpp
new file mode 100644
--- /dev/null
+++ b/Engine/Plugins/Developer/Sandbox/SandboxedEditing/Source/SandboxedEditing/Private/Widgets/SSandboxStatusBarWidget.cpp
@@ -0,0 +1,127 @@
+﻿// Copyright Epic Games, Inc. All Rights Reserved.
+
+#include "SSandboxStatusBarWidget.h"
+
+#include "Framework/Data/SandboxItem.h"
+#include "Framework/ViewModels/SandboxSystemViewModel.h"
+#include "SActionButton.h"
+#include "SandboxedEditingCommands.h"
+#include "Styling/AppStyle.h"
+#include "ToolMenus.h"
+#include "ToolWidgetsSlateTypes.h"
+#include "Framework/Notifications.h"
+#include "Widgets/Input/SComboButton.h"
+#include "Widgets/SBoxPanel.h"
+#include "Widgets/Text/STextBlock.h"
+
+#define LOCTEXT_NAMESPACE "SSandboxStatusBarWidget"
+
+namespace UE::SandboxedEditing
+{
+void SSandboxStatusBarWidget::Construct(
+	const FArguments& InArgs,
+	const TSharedRef<FSandboxSystemViewModel>& InViewModel, const TSharedRef<FUICommandList>& InCommandList
+	)
+{
+	ViewModel = InViewModel;
+	CommandList = InCommandList;
+	
+	ChildSlot
+	[
+		SNew(SComboButton)
+		.ContentPadding(FMargin(6.0f, 0.0f))
+		.MenuPlacement(MenuPlacement_AboveAnchor)
+		.ComboButtonStyle(&FAppStyle::Get().GetWidgetStyle<FComboButtonStyle>("SimpleComboButton"))
+		.ButtonContent()
+		[
+			SNew(SHorizontalBox)
+	
+			+ SHorizontalBox::Slot()
+			.AutoWidth()
+			.VAlign(VAlign_Center)
+			.Padding(0, 0, 10, 0)
+			[
+				SNew(STextBlock)
+				.Text(this, &SSandboxStatusBarWidget::GetTitleText)
+				.ToolTipText(LOCTEXT("AssetManagementStatusBarToolTip", "Tools for altering the state of assets."))
+			]
+		]
+		.OnGetMenuContent(FOnGetContent::CreateRaw(this, &SSandboxStatusBarWidget::CreateStatusBarMenu))
+	];
+}
+
+FText SSandboxStatusBarWidget::GetTitleText() const
+{
+	const FString SandboxName = ViewModel->GetActiveSandboxName();
+	if (!SandboxName.IsEmpty())
+	{
+		FTextBuilder TextBuilder;
+		TextBuilder.AppendLineFormat(LOCTEXT("LoadedSandbox", "{0} (Sandbox)"), FText::FromString(*SandboxName));
+		return TextBuilder.ToText();
+	}
+	return LOCTEXT("StatusBarName", "Sandbox");
+}
+
+TSharedRef<SWidget> SSandboxStatusBarWidget::CreateStatusBarMenu()
+{
+	UToolMenu* Menu = UToolMenus::Get()->RegisterMenu("StatusBar.ToolBar.SandboxedEditing", NAME_None, EMultiBoxType::Menu, false);
+
+	{
+		FToolMenuSection& Section = Menu->AddSection("SandboxedEditingSection", LOCTEXT("SandboxedEditing.Sandbox", "Sandbox"));
+
+		Section.AddMenuEntry(FSandboxedEditingCommands::Get().OpenCreateNewSandboxDialog);
+		Section.AddMenuEntry(FSandboxedEditingCommands::Get().LeaveSandbox);
+		Section.AddMenuEntry(FSandboxedEditingCommands::Get().PersistAll);
+		Section.AddMenuEntry(FSandboxedEditingCommands::Get().DiscardAll);
+
+		BuildSandboxesMenu(Menu);
+	}
+	
+	return UToolMenus::Get()->GenerateWidget("StatusBar.ToolBar.SandboxedEditing", FToolMenuContext(CommandList));
+}
+
+void SSandboxStatusBarWidget::BuildSandboxesMenu(UToolMenu* InMenuBuilder)
+{
+	FToolMenuSection& Section = InMenuBuilder->AddSection("Load Sandboxes", LOCTEXT("LoadSandboxes", "Load Sandbox"));
+	Section.AddDynamicEntry("Existing Sandboxes",
+		FNewToolMenuSectionDelegate::CreateLambda([this](FToolMenuSection& InSection)
+	{
+		const TArray<FSandboxItem> Sandboxes = ViewModel->GetKnownSandboxes();
+
+		TArray<FString> DuplicateNameDetection;
+		for (const FSandboxItem& Sandbox : Sandboxes)
+		{
+			// Tool menus want unique names. It could happen that the user goes into the file system and manually edits the name entry in the JSON.
+			const FName MenuName(
+				FString::Printf(TEXT("%s_%s"),* Sandbox.Name, *FGuid::NewDeterministicGuid(Sandbox.SandboxRoot).ToString())
+				);
+			
+			const FToolMenuEntry Entry = FToolMenuEntry::InitMenuEntry(
+				MenuName,
+				FText::FromString(Sandbox.Name),
+				FText::FromString(Sandbox.Description),
+				FSlateIcon(),
+				FUIAction(
+					FExecuteAction::CreateLambda([this, Sandbox]
+					{
+						if (ViewModel->LoadSandbox(Sandbox.SandboxRoot))
+						{
+							ShowLoadedSandbox(Sandbox.Name);
+						}
+						else
+						{
+							ShowFailedToLoadSandbox(Sandbox.Name);
+						}
+					}),
+					FCanExecuteAction::CreateLambda([this, Path = Sandbox.SandboxRoot]{ return ViewModel->CanLoadSandbox(Path); }),
+					FIsActionChecked::CreateLambda([this, Path = Sandbox.SandboxRoot]{ return ViewModel->IsActiveSandbox(Path); })
+					),
+				EUserInterfaceActionType::RadioButton
+				);
+			InSection.AddEntry(Entry);
+		}
+	}));
+}
+}
+
+#undef LOCTEXT_NAMESPACE
\ No newline at end of file
diff --git a/Engine/Plugins/Developer/Sandbox/SandboxedEditing/Source/SandboxedEditing/Private/Widgets/SSandboxStatusBarWidget.h b/Engine/Plugins/Developer/Sandbox/SandboxedEditing/Source/SandboxedEditing/Private/Widgets/SSandboxStatusBarWidget.h
new file mode 100644
--- /dev/null
+++ b/Engine/Plugins/Developer/Sandbox/SandboxedEditing/Source/SandboxedEditing/Private/Widgets/SSandboxStatusBarWidget.h
@@ -0,0 +1,38 @@
+﻿// Copyright Epic Games, Inc. All Rights Reserved.
+
+#pragma once
+
+#include "ToolMenu.h"
+#include "Widgets/SCompoundWidget.h"
+#include "Widgets/DeclarativeSyntaxSupport.h"
+
+namespace UE::SandboxedEditing
+{
+class FSandboxSystemViewModel;
+
+/** Root widget that is placed in the status bar. */
+class SSandboxStatusBarWidget : public SCompoundWidget
+{
+public:
+
+	SLATE_BEGIN_ARGS(SSandboxStatusBarWidget){}
+	SLATE_END_ARGS()
+
+	void Construct(const FArguments& InArgs, const TSharedRef<FSandboxSystemViewModel>& InViewModel, const TSharedRef<FUICommandList>& InCommandList);
+
+private:
+
+	/** The model we're bound to. */
+	TSharedPtr<FSandboxSystemViewModel> ViewModel;
+	/** The commands needed to create the menu. */
+	TSharedPtr<FUICommandList> CommandList;
+
+	/** @return Text to display in the combo button content */
+	FText GetTitleText() const;
+	/** @return Content to show when combo button is opened */
+	TSharedRef<SWidget> CreateStatusBarMenu();
+
+	/** Fills the submenu showing the available sandboxes. */
+	void BuildSandboxesMenu(UToolMenu* InMenuBuilder);
+};
+}
diff --git a/Engine/Plugins/Developer/Sandbox/SandboxedEditing/Source/SandboxedEditing/Public/ISandboxEditingModule.h b/Engine/Plugins/Developer/Sandbox/SandboxedEditing/Source/SandboxedEditing/Public/ISandboxEditingModule.h
new file mode 100644
--- /dev/null
+++ b/Engine/Plugins/Developer/Sandbox/SandboxedEditing/Source/SandboxedEditing/Public/ISandboxEditingModule.h
@@ -0,0 +1,29 @@
+// Copyright Epic Games, Inc. All Rights Reserved.
+
+#pragma once
+
+#include "Modules/ModuleInterface.h"
+#include "Modules/ModuleManager.h"
+
+namespace UE::SandboxedEditing
+{
+class ISandboxedEditingModule : public IModuleInterface
+{
+public:
+
+	/**
+	 * Singleton-like access to this module's interface.  This is just for convenience!
+	 * Beware of calling this during the shutdown phase, though.  Your module might have been unloaded already.
+	 *
+	 * @return Returns singleton instance, loading the module on demand if needed
+	 */
+	static ISandboxedEditingModule& Get() { return FModuleManager::LoadModuleChecked<ISandboxedEditingModule>("SandboxedEditing"); }
+
+	/**
+	 * Checks to see if this module is loaded and ready.  It is only valid to call Get() during shutdown if IsAvailable() returns true.
+	 *
+	 * @return True if the module is loaded and ready to use
+	 */
+	static bool IsAvailable() { return FModuleManager::Get().IsModuleLoaded("SandboxedEditing"); }
+};
+}
\ No newline at end of file
diff --git a/Engine/Plugins/Developer/Sandbox/SandboxedEditing/Source/SandboxedEditing/SandboxedEditing.Build.cs b/Engine/Plugins/Developer/Sandbox/SandboxedEditing/Source/SandboxedEditing/SandboxedEditing.Build.cs
new file mode 100644
--- /dev/null
+++ b/Engine/Plugins/Developer/Sandbox/SandboxedEditing/Source/SandboxedEditing/SandboxedEditing.Build.cs
@@ -0,0 +1,30 @@
+// Copyright Epic Games, Inc. All Rights Reserved.
+
+namespace UnrealBuildTool.Rules
+{
+	public class SandboxedEditing : ModuleRules
+	{
+		public SandboxedEditing(ReadOnlyTargetRules Target) : base(Target)
+		{
+			PublicDependencyModuleNames.AddRange(
+				new string[]
+				{
+					"Core",
+					"CoreUObject",  
+					"FileSandboxCore"
+				}
+			);
+			
+			PrivateDependencyModuleNames.AddRange(
+				new string[]
+				{
+					"InputCore",
+					"Slate", 
+					"SlateCore", 
+					"ToolMenus", 
+					"ToolWidgets",
+				}
+			);
+		}
+	}
+}
