File Sandbox: Remove unused function

#Synced-CL 47637010

diff --git a/Engine/Plugins/Developer/Sandbox/SandboxedEditing/Source/SandboxedEditing/Private/Widgets/Dialogues/SCreateNewSandbox.h b/Engine/Plugins/Developer/Sandbox/SandboxedEditing/Source/SandboxedEditing/Private/Widgets/Dialogues/SCreateNewSandbox.h
--- a/Engine/Plugins/Developer/Sandbox/SandboxedEditing/Source/SandboxedEditing/Private/Widgets/Dialogues/SCreateNewSandbox.h
+++ b/Engine/Plugins/Developer/Sandbox/SandboxedEditing/Source/SandboxedEditing/Private/Widgets/Dialogues/SCreateNewSandbox.h
@@ -60,9 +60,6 @@ private:
 	/** @return Tooltip text to show on create button */
 	FText GetCreateButtonToolTipText() const;
 
-	/** @return Whether the sandbox can be created. */
-	bool CanCreate(FText* OutReason);
-
 	/** Updates the error state when the name changes. */
 	void OnNameChanged(const FText& InNewName);
 	
