Sandboxed Editing: Fix localization build issue

#Synced-CL 47637010

diff --git a/Engine/Plugins/Developer/Sandbox/SandboxedEditing/Source/SandboxedEditing/Private/Features/StatusBar/StatusBarFeature.cpp b/Engine/Plugins/Developer/Sandbox/SandboxedEditing/Source/SandboxedEditing/Private/Features/StatusBar/StatusBarFeature.cpp
--- a/Engine/Plugins/Developer/Sandbox/SandboxedEditing/Source/SandboxedEditing/Private/Features/StatusBar/StatusBarFeature.cpp
+++ b/Engine/Plugins/Developer/Sandbox/SandboxedEditing/Source/SandboxedEditing/Private/Features/StatusBar/StatusBarFeature.cpp
@@ -4,6 +4,7 @@
 
 #include "ToolMenus.h"
 #include "Features/StatusBar/Widgets/SSandboxStatusBarWidget.h"
+#include "Framework/Application/SlateApplication.h"
 
 namespace UE::SandboxedEditing
 {
@@ -16,6 +17,11 @@ static TSharedRef<SWidget> CreateStatusBarWidget(const TSharedRef<FSandboxSystem
 
 static void RegisterStatusBar(const TSharedRef<FSandboxSystemModel>& InViewModel, const FStatusBarCommandMappings& InCommandMappings)
 {
+	if (!FSlateApplication::IsInitialized())
+	{
+		return;
+	}
+	
 	UToolMenu* Menu = UToolMenus::Get()->ExtendMenu(TEXT("LevelEditor.StatusBar.ToolBar"));
 	FToolMenuSection& AssetManagementSection = Menu->AddSection(
 		TEXT("Asset Management"), FText::GetEmpty(), FToolMenuInsert(NAME_None, EToolMenuInsertType::First)
@@ -31,9 +37,6 @@ FStatusBarFeature::FStatusBarFeature(const TSharedRef<FSandboxSystemModel>& InSa
 	: SandboxModel(InSandboxModel)
 	, CommandMappings(SandboxModel)
 {
-	FCoreDelegates::OnPostEngineInit.AddLambda([this]
-	{
-		StatusBarDetail::RegisterStatusBar(SandboxModel, CommandMappings);
-	});
+	StatusBarDetail::RegisterStatusBar(SandboxModel, CommandMappings);
 }
 }
