Sandboxed Editing: Add a sandbox browser.

Features:
- Window > Sandbox entry
- Browser shows all sandboxes in Intermediate/Sandboxes folder as list view.
- List view displays name, version, last modified
- Description can be edited using expandable section below
- Browser can be searched
- Upon entering sandbox, widget changes UI to display info about current sandbox
- For now only session name and description exist here

Missing features: Persisting changes. Renaming sandboxes. Sorting the list view. Proper icons.

#Synced-CL 47637010

diff --git a/Engine/Plugins/Developer/Sandbox/SandboxedEditing/SandboxedEditing.uplugin b/Engine/Plugins/Developer/Sandbox/SandboxedEditing/SandboxedEditing.uplugin
--- a/Engine/Plugins/Developer/Sandbox/SandboxedEditing/SandboxedEditing.uplugin
+++ b/Engine/Plugins/Developer/Sandbox/SandboxedEditing/SandboxedEditing.uplugin
@@ -13,7 +13,7 @@
 	"EnabledByDefault": false,
 	"CanContainContent": false,
 	"Hidden": true,
-	"IsBetaVersion": true,
+	"IsBetaVersion": false,
 	"Installed": false,
 	"Modules": [
 		{
diff --git a/Engine/Plugins/Developer/Sandbox/SandboxedEditing/Source/SandboxedEditing/Private/Features/Browser/BrowserFeature.cpp b/Engine/Plugins/Developer/Sandbox/SandboxedEditing/Source/SandboxedEditing/Private/Features/Browser/BrowserFeature.cpp
new file mode 100644
--- /dev/null
+++ b/Engine/Plugins/Developer/Sandbox/SandboxedEditing/Source/SandboxedEditing/Private/Features/Browser/BrowserFeature.cpp
@@ -0,0 +1,90 @@
+﻿// Copyright Epic Games, Inc. All Rights Reserved.
+
+#include "BrowserFeature.h"
+
+#include "Framework/Docking/TabManager.h"
+#include "SandboxedEditingStyle.h"
+#include "ViewModels/List/SandboxListViewModel.h"
+#include "Widgets/Docking/SDockTab.h"
+#include "Widgets/SSandboxedEditingRoot.h"
+#include "WorkspaceMenuStructure.h"
+#include "WorkspaceMenuStructureModule.h"
+#include "Framework/Application/SlateApplication.h"
+#include "ViewModels/SandboxControlsViewModel.h"
+
+#define LOCTEXT_NAMESPACE "FBrowserFeature"
+
+namespace UE::SandboxedEditing
+{
+namespace BrowserDetail
+{
+const FLazyName SandboxedEditingTabName(TEXT("SandboxedEditingTab"));
+
+static TMap<FName, TSharedRef<ISandboxColumnBehavior>> TransformToBehaviours(const TMap<FName, FSandboxColumnInfo>& InColumns)
+{
+	TMap<FName, TSharedRef<ISandboxColumnBehavior>> Result;
+	Algo::Transform(InColumns, Result, [](const TPair<FName, FSandboxColumnInfo>& Pair)
+	{
+		return TPair<FName, TSharedRef<ISandboxColumnBehavior>>{ Pair.Key, Pair.Value.Behavior };
+	});
+	return Result;
+}
+
+static TMap<FName, TSharedRef<ISandboxWidgetFactory>> TransformToWidgetFactories(const TMap<FName, FSandboxColumnInfo>& InColumns)
+{
+	TMap<FName, TSharedRef<ISandboxWidgetFactory>> Result;
+	Algo::Transform(InColumns, Result, [](const TPair<FName, FSandboxColumnInfo>& Pair)
+	{
+		return TPair<FName, TSharedRef<ISandboxWidgetFactory>>{ Pair.Key, Pair.Value.WidgetFactory };
+	});
+	return Result;
+}
+}
+	
+FBrowserFeature::FBrowserFeature(const TSharedRef<FSandboxSystemModel>& InSandboxModel)
+	: SandboxModel(InSandboxModel)
+	, BrowserColumns(MakeBuiltInColumns())
+	, ViewModels(BrowserDetail::TransformToBehaviours(BrowserColumns), SandboxModel)
+	, CommandBindings(MakeShared<FBrowserCommandBindings>(ViewModels))
+{
+	RegisterTabSpawner();
+}
+
+FBrowserFeature::~FBrowserFeature()
+{
+	UnregisterTabSpawner();
+}
+
+void FBrowserFeature::RegisterTabSpawner()
+{
+	FTabSpawnerEntry& BrowserSpawnerEntry = FGlobalTabmanager::Get()->RegisterNomadTabSpawner(
+		BrowserDetail::SandboxedEditingTabName, FOnSpawnTab::CreateRaw(this, &FBrowserFeature::SpawnBrowserTab)
+		);
+	
+	BrowserSpawnerEntry
+	.SetIcon(FSlateIcon(FSandboxedEditingStyle::Get().GetStyleSetName(), TEXT("SandboxedEditing.Browser.Icon")))
+	.SetDisplayName(LOCTEXT("BrowserTabTitle", "Sandboxes"))
+	.SetTooltipText(LOCTEXT("BrowserTooltipText", "Open the Sandbox browser"))
+	.SetMenuType(ETabSpawnerMenuType::Enabled);
+	
+	BrowserSpawnerEntry.SetGroup(WorkspaceMenu::GetMenuStructure().GetLevelEditorCategory());
+}
+
+void FBrowserFeature::UnregisterTabSpawner()
+{
+	FGlobalTabmanager::Get()->UnregisterNomadTabSpawner(BrowserDetail::SandboxedEditingTabName);
+}
+
+TSharedRef<SDockTab> FBrowserFeature::SpawnBrowserTab(const FSpawnTabArgs& InArgs)
+{
+	return SNew(SDockTab)
+		.TabRole(NomadTab)
+		[
+			SNew(SSandboxedEditingRoot, ViewModels)
+			.CommandList(CommandBindings->GetCommandList())
+			.ColumnFactories(BrowserDetail::TransformToWidgetFactories(BrowserColumns))
+		];
+}
+}
+
+#undef LOCTEXT_NAMESPACE
\ No newline at end of file
diff --git a/Engine/Plugins/Developer/Sandbox/SandboxedEditing/Source/SandboxedEditing/Private/Features/Browser/BrowserFeature.h b/Engine/Plugins/Developer/Sandbox/SandboxedEditing/Source/SandboxedEditing/Private/Features/Browser/BrowserFeature.h
new file mode 100644
--- /dev/null
+++ b/Engine/Plugins/Developer/Sandbox/SandboxedEditing/Source/SandboxedEditing/Private/Features/Browser/BrowserFeature.h
@@ -0,0 +1,49 @@
+﻿// Copyright Epic Games, Inc. All Rights Reserved.
+
+#pragma once
+
+#include "Commands/BrowserCommandBindings.h"
+#include "Templates/SharedPointer.h"
+#include "Templates/UnrealTemplate.h"
+#include "ViewModels/BrowserViewModels.h"
+
+class SDockTab;
+class FSpawnTabArgs;
+
+namespace UE::SandboxedEditing
+{
+class FActiveSandboxDetailsViewModel;
+class FActiveSandboxTrackerViewModel;
+class SSandboxedEditingRoot;
+class FSandboxControlsViewModel;
+class FSandboxListViewModel;
+class FSandboxSystemModel;
+
+/** Summons a nomad tab through which the user can browse sandboxes. */
+class FBrowserFeature : public FNoncopyable
+{
+public:
+
+	explicit FBrowserFeature(const TSharedRef<FSandboxSystemModel>& InSandboxModel);
+	~FBrowserFeature();
+	
+private:
+	
+	/** The shared sandbox model. */
+	const TSharedRef<FSandboxSystemModel> SandboxModel;
+	
+	/** The columns used by the browser widget. */
+	const TMap<FName, FSandboxColumnInfo> BrowserColumns;
+	/** Holds all view models relevant for the browser. */
+	const FBrowserViewModels ViewModels;
+	
+	/** Binds user commands to models.*/
+	const TSharedRef<FBrowserCommandBindings> CommandBindings;
+	
+	void RegisterTabSpawner();
+	void UnregisterTabSpawner();
+	TSharedRef<SDockTab> SpawnBrowserTab(const FSpawnTabArgs& InArgs);
+};
+}
+
+
diff --git a/Engine/Plugins/Developer/Sandbox/SandboxedEditing/Source/SandboxedEditing/Private/Features/Browser/Commands/BrowserCommandBindings.cpp b/Engine/Plugins/Developer/Sandbox/SandboxedEditing/Source/SandboxedEditing/Private/Features/Browser/Commands/BrowserCommandBindings.cpp
new file mode 100644
--- /dev/null
+++ b/Engine/Plugins/Developer/Sandbox/SandboxedEditing/Source/SandboxedEditing/Private/Features/Browser/Commands/BrowserCommandBindings.cpp
@@ -0,0 +1,57 @@
+﻿// Copyright Epic Games, Inc. All Rights Reserved.
+
+#include "BrowserCommandBindings.h"
+
+#include "BrowserCommands.h"
+#include "Features/Browser/ViewModels/BrowserViewModels.h"
+#include "Features/Browser/ViewModels/SandboxControlsViewModel.h"
+#include "Features/Browser/ViewModels/Active/ActiveSandboxTrackerViewModel.h"
+#include "Framework/Commands/GenericCommands.h"
+
+namespace UE::SandboxedEditing
+{
+FBrowserCommandBindings::FBrowserCommandBindings(const FBrowserViewModels& InViewModels)
+	: ControlsViewModel(InViewModels.ControlsViewModel)
+	, ActiveSandboxViewModel(InViewModels.ActiveSandboxDetailsViewModel)
+	, CommandList(MakeShared<FUICommandList>())
+{
+	FBrowserCommands::Register();
+	
+	FBrowserCommands& BrowserCommands = FBrowserCommands::Get();
+	CommandList->MapAction(BrowserCommands.CreateNewSandbox, 
+		FExecuteAction::CreateSP(ControlsViewModel, &FSandboxControlsViewModel::StartCreationWorkflow), 
+		FCanExecuteAction::CreateSP(ControlsViewModel, &FSandboxControlsViewModel::CanStartCreationWorkflow)
+	);
+	CommandList->MapAction(BrowserCommands.Cancel, FExecuteAction::CreateRaw(this, &FBrowserCommandBindings::HandleCancel));
+	CommandList->MapAction(BrowserCommands.LeaveSandbox, 
+		FExecuteAction::CreateRaw(this, &FBrowserCommandBindings::HandleLeaveSandbox),
+		FCanExecuteAction::CreateRaw(this, &FBrowserCommandBindings::CanLeaveSandbox)
+		);
+}
+
+FBrowserCommandBindings::~FBrowserCommandBindings()
+{
+	FBrowserCommands::Unregister();
+}
+
+void FBrowserCommandBindings::HandleCancel() const
+{
+	if (FSandboxCreationWorkflow* CreationWorkflow = ControlsViewModel->GetCurrentCreationWorkflow())
+	{
+		CreationWorkflow->Cancel();
+	}
+}
+
+void FBrowserCommandBindings::HandleLeaveSandbox()
+{
+	if (CanLeaveSandbox())
+	{
+		ActiveSandboxViewModel->LeaveSandbox();
+	}
+}
+
+bool FBrowserCommandBindings::CanLeaveSandbox()
+{
+	return ActiveSandboxViewModel->HasActiveSandbox();
+}
+}
diff --git a/Engine/Plugins/Developer/Sandbox/SandboxedEditing/Source/SandboxedEditing/Private/Features/Browser/Commands/BrowserCommandBindings.h b/Engine/Plugins/Developer/Sandbox/SandboxedEditing/Source/SandboxedEditing/Private/Features/Browser/Commands/BrowserCommandBindings.h
new file mode 100644
--- /dev/null
+++ b/Engine/Plugins/Developer/Sandbox/SandboxedEditing/Source/SandboxedEditing/Private/Features/Browser/Commands/BrowserCommandBindings.h
@@ -0,0 +1,45 @@
+﻿// Copyright Epic Games, Inc. All Rights Reserved.
+
+#pragma once
+
+#include "Framework/Commands/UICommandList.h"
+#include "Templates/SharedPointer.h"
+#include "Templates/UnrealTemplate.h"
+
+namespace UE::SandboxedEditing
+{
+class FActiveSandboxDetailsViewModel;
+class FSandboxControlsViewModel;
+struct FBrowserViewModels;
+
+/** Binds the commands relevant for the browser feature */
+class FBrowserCommandBindings : public FNoncopyable
+{
+public:
+	
+	explicit FBrowserCommandBindings(const FBrowserViewModels& InViewModels);
+	~FBrowserCommandBindings();
+	
+	const TSharedRef<FUICommandList>& GetCommandList() const { return CommandList; }
+	
+private:
+	
+	/** Used for editing sandbox info. */
+	const TSharedRef<FSandboxControlsViewModel> ControlsViewModel;
+	/** Used to interact with the active sandbox. */
+	const TSharedRef<FActiveSandboxDetailsViewModel> ActiveSandboxViewModel;
+	
+	/** Holds our command bindings */
+	const TSharedRef<FUICommandList> CommandList;
+	
+	/** Invoked when the delete command is issued. */
+	FSimpleMulticastDelegate OnDeleteSelection;
+	
+	/** Cancels any active operation, such as renaming or creating a sandbox. */
+	void HandleCancel() const;
+	
+	void HandleLeaveSandbox();
+	bool CanLeaveSandbox();
+};
+}
+
diff --git a/Engine/Plugins/Developer/Sandbox/SandboxedEditing/Source/SandboxedEditing/Private/Features/Browser/Commands/BrowserCommands.cpp b/Engine/Plugins/Developer/Sandbox/SandboxedEditing/Source/SandboxedEditing/Private/Features/Browser/Commands/BrowserCommands.cpp
new file mode 100644
--- /dev/null
+++ b/Engine/Plugins/Developer/Sandbox/SandboxedEditing/Source/SandboxedEditing/Private/Features/Browser/Commands/BrowserCommands.cpp
@@ -0,0 +1,29 @@
+﻿// Copyright Epic Games, Inc. All Rights Reserved.
+
+#include "BrowserCommands.h"
+
+#include "SandboxedEditingStyle.h"
+
+#define LOCTEXT_NAMESPACE "FBrowserCommands"
+
+namespace UE::SandboxedEditing
+{
+FBrowserCommands::FBrowserCommands()
+	: TCommands(
+		TEXT("Browser"), 
+		LOCTEXT("Description", "Commands relevant to the sandbox browser"), 
+		NAME_None, 
+		FSandboxedEditingStyle::Get().GetStyleSetName()
+		)
+{}
+
+void FBrowserCommands::RegisterCommands()
+{
+	UI_COMMAND(CreateNewSandbox, "Create New Sandbox", "Creates a new sandbox", EUserInterfaceActionType::Button, FInputChord());
+	UI_COMMAND(Cancel, "Cancel", "Cancel the current operation, such as renaming or creating a sandbox.", EUserInterfaceActionType::None, FInputChord(EKeys::Escape));
+	UI_COMMAND(LeaveSandbox, "Leave Sandbox", "Leaves the current sandbox.", EUserInterfaceActionType::Button, FInputChord());
+	
+}
+}
+
+#undef LOCTEXT_NAMESPACE
\ No newline at end of file
diff --git a/Engine/Plugins/Developer/Sandbox/SandboxedEditing/Source/SandboxedEditing/Private/Features/Browser/Commands/BrowserCommands.h b/Engine/Plugins/Developer/Sandbox/SandboxedEditing/Source/SandboxedEditing/Private/Features/Browser/Commands/BrowserCommands.h
new file mode 100644
--- /dev/null
+++ b/Engine/Plugins/Developer/Sandbox/SandboxedEditing/Source/SandboxedEditing/Private/Features/Browser/Commands/BrowserCommands.h
@@ -0,0 +1,30 @@
+﻿// Copyright Epic Games, Inc. All Rights Reserved.
+
+#pragma once
+
+#include "Framework/Commands/Commands.h"
+
+namespace UE::SandboxedEditing
+{
+/** Commands used in the browser feature. */
+class FBrowserCommands : public TCommands<FBrowserCommands>
+{
+public:
+	
+	FBrowserCommands();
+	
+	//~ Begin TCommands Interface
+	virtual void RegisterCommands() override;
+	//~ End TCommands Interface
+	
+	/** Starts the workflow for creating a new sandbox. */
+	TSharedPtr<FUICommandInfo> CreateNewSandbox;
+	
+	/** Cancel the current operation, such as renaming or creating a sandbox. */
+	TSharedPtr<FUICommandInfo> Cancel;
+	
+	/** Leaves the current sandbox. */
+	TSharedPtr<FUICommandInfo> LeaveSandbox;
+};
+}
+
diff --git a/Engine/Plugins/Developer/Sandbox/SandboxedEditing/Source/SandboxedEditing/Private/Features/Browser/ViewModels/Active/ActiveSandboxDetailsViewModel.cpp b/Engine/Plugins/Developer/Sandbox/SandboxedEditing/Source/SandboxedEditing/Private/Features/Browser/ViewModels/Active/ActiveSandboxDetailsViewModel.cpp
new file mode 100644
--- /dev/null
+++ b/Engine/Plugins/Developer/Sandbox/SandboxedEditing/Source/SandboxedEditing/Private/Features/Browser/ViewModels/Active/ActiveSandboxDetailsViewModel.cpp
@@ -0,0 +1,32 @@
+﻿// Copyright Epic Games, Inc. All Rights Reserved.
+
+#include "ActiveSandboxDetailsViewModel.h"
+
+#include "Framework/Models/SandboxSystemModel.h"
+
+namespace UE::SandboxedEditing
+{
+FActiveSandboxDetailsViewModel::FActiveSandboxDetailsViewModel(const TSharedRef<FSandboxSystemModel>& InModel)
+	: Model(InModel)
+{}
+
+bool FActiveSandboxDetailsViewModel::HasActiveSandbox() const
+{
+	return Model->HasActiveSandbox();
+}
+
+void FActiveSandboxDetailsViewModel::LeaveSandbox()
+{
+	Model->LeaveSandbox();
+}
+
+FString FActiveSandboxDetailsViewModel::GetSandboxName() const
+{
+	return Model->GetActiveSandboxName();
+}
+
+FString FActiveSandboxDetailsViewModel::GetSandboxPath() const
+{
+	return Model->GetActiveSandboxPath();
+}
+}
diff --git a/Engine/Plugins/Developer/Sandbox/SandboxedEditing/Source/SandboxedEditing/Private/Features/Browser/ViewModels/Active/ActiveSandboxDetailsViewModel.h b/Engine/Plugins/Developer/Sandbox/SandboxedEditing/Source/SandboxedEditing/Private/Features/Browser/ViewModels/Active/ActiveSandboxDetailsViewModel.h
new file mode 100644
--- /dev/null
+++ b/Engine/Plugins/Developer/Sandbox/SandboxedEditing/Source/SandboxedEditing/Private/Features/Browser/ViewModels/Active/ActiveSandboxDetailsViewModel.h
@@ -0,0 +1,37 @@
+﻿// Copyright Epic Games, Inc. All Rights Reserved.
+
+#pragma once
+
+#include "Delegates/Delegate.h"
+#include "Templates/SharedPointer.h"
+#include "Templates/UnrealTemplate.h"
+
+namespace UE::SandboxedEditing
+{
+class FSandboxSystemModel;
+
+/** Information about the active sandbox. */
+class FActiveSandboxDetailsViewModel : public FNoncopyable
+{
+public:
+	
+	explicit FActiveSandboxDetailsViewModel(const TSharedRef<FSandboxSystemModel>& InModel);
+	
+	/** @return Whether currently in an active sandbox */
+	bool HasActiveSandbox() const;
+	
+	/** Leaves the active sandbox. */
+	void LeaveSandbox();
+	
+	/** @return Name of the active sandbox. */
+	FString GetSandboxName() const;
+	/** @return Path to active sandbox root. */
+	FString GetSandboxPath() const;
+
+private:
+	
+	/** The model used to interact with the sandbox system. */
+	const TSharedRef<FSandboxSystemModel> Model;
+};
+}
+
diff --git a/Engine/Plugins/Developer/Sandbox/SandboxedEditing/Source/SandboxedEditing/Private/Features/Browser/ViewModels/Active/ActiveSandboxTrackerViewModel.cpp b/Engine/Plugins/Developer/Sandbox/SandboxedEditing/Source/SandboxedEditing/Private/Features/Browser/ViewModels/Active/ActiveSandboxTrackerViewModel.cpp
new file mode 100644
--- /dev/null
+++ b/Engine/Plugins/Developer/Sandbox/SandboxedEditing/Source/SandboxedEditing/Private/Features/Browser/ViewModels/Active/ActiveSandboxTrackerViewModel.cpp
@@ -0,0 +1,36 @@
+﻿// Copyright Epic Games, Inc. All Rights Reserved.
+
+#include "ActiveSandboxTrackerViewModel.h"
+
+#include "Framework/Models/SandboxSystemModel.h"
+
+namespace UE::SandboxedEditing
+{
+FActiveSandboxTrackerViewModel::FActiveSandboxTrackerViewModel(const TSharedRef<FSandboxSystemModel>& InModel)
+	: Model(InModel)
+{
+	Model->OnLoadSandbox().AddRaw(this, &FActiveSandboxTrackerViewModel::HandleSandboxedLoaded);
+	Model->OnLeaveSandbox().AddRaw(this, &FActiveSandboxTrackerViewModel::HandleSandboxLeft);
+}
+
+FActiveSandboxTrackerViewModel::~FActiveSandboxTrackerViewModel()
+{
+	Model->OnLoadSandbox().RemoveAll(this);
+	Model->OnLeaveSandbox().RemoveAll(this);
+}
+
+bool FActiveSandboxTrackerViewModel::HasActiveSandbox() const
+{
+	return Model->HasActiveSandbox();
+}
+
+void FActiveSandboxTrackerViewModel::HandleSandboxedLoaded()
+{
+	OnLoadSandboxDelegate.Broadcast();
+}
+
+void FActiveSandboxTrackerViewModel::HandleSandboxLeft()
+{
+	OnLeaveSandboxDelegate.Broadcast();
+}
+}
diff --git a/Engine/Plugins/Developer/Sandbox/SandboxedEditing/Source/SandboxedEditing/Private/Features/Browser/ViewModels/Active/ActiveSandboxTrackerViewModel.h b/Engine/Plugins/Developer/Sandbox/SandboxedEditing/Source/SandboxedEditing/Private/Features/Browser/ViewModels/Active/ActiveSandboxTrackerViewModel.h
new file mode 100644
--- /dev/null
+++ b/Engine/Plugins/Developer/Sandbox/SandboxedEditing/Source/SandboxedEditing/Private/Features/Browser/ViewModels/Active/ActiveSandboxTrackerViewModel.h
@@ -0,0 +1,46 @@
+﻿// Copyright Epic Games, Inc. All Rights Reserved.
+
+#pragma once
+
+#include "ActiveSandboxDetailsViewModel.h"
+#include "Delegates/Delegate.h"
+#include "Templates/SharedPointer.h"
+#include "Templates/UnrealTemplate.h"
+
+namespace UE::SandboxedEditing
+{
+class FSandboxSystemModel;
+
+/** Knows when the active sandbox changes. */
+class FActiveSandboxTrackerViewModel : public FNoncopyable
+{
+public:
+	
+	explicit FActiveSandboxTrackerViewModel(const TSharedRef<FSandboxSystemModel>& InModel);
+	~FActiveSandboxTrackerViewModel();
+	
+	/** @return Whether there is an active sandbox */
+	bool HasActiveSandbox() const;
+	
+	/** Invoked when the engine has entered an active sandbox. */
+	FSimpleMulticastDelegate& OnLoadSandbox() { return OnLoadSandboxDelegate; }
+	/** Invoked when the active has stopped being in an active sandbox. */
+	FSimpleMulticastDelegate& OnLeaveSandbox() { return OnLeaveSandboxDelegate; }
+	
+private:
+	
+	/** The model used to interact with the sandbox system. */
+	const TSharedRef<FSandboxSystemModel> Model;
+	
+	/** Invoked when a new sandbox is loaded. */
+	FSimpleMulticastDelegate OnLoadSandboxDelegate;
+	/** Invoked when the current sandbox is left. */
+	FSimpleMulticastDelegate OnLeaveSandboxDelegate;
+	
+	/** Creates the view model and invokes OnLoadSandboxDelegate. */
+	void HandleSandboxedLoaded();
+	/** Cleans up the view model and invokes OnLeaveSandboxDelegate. */
+	void HandleSandboxLeft();
+};
+}
+
diff --git a/Engine/Plugins/Developer/Sandbox/SandboxedEditing/Source/SandboxedEditing/Private/Features/Browser/ViewModels/BrowserViewModels.h b/Engine/Plugins/Developer/Sandbox/SandboxedEditing/Source/SandboxedEditing/Private/Features/Browser/ViewModels/BrowserViewModels.h
new file mode 100644
--- /dev/null
+++ b/Engine/Plugins/Developer/Sandbox/SandboxedEditing/Source/SandboxedEditing/Private/Features/Browser/ViewModels/BrowserViewModels.h
@@ -0,0 +1,50 @@
+﻿// Copyright Epic Games, Inc. All Rights Reserved.
+
+#pragma once
+
+#include "Active/ActiveSandboxDetailsViewModel.h"
+#include "Active/ActiveSandboxTrackerViewModel.h"
+#include "List/FilterSandboxViewModel.h"
+#include "List/SandboxListViewModel.h"
+#include "SandboxControlsViewModel.h"
+#include "SandboxMetaDataViewModel.h"
+#include "Templates/SharedPointer.h"
+
+namespace UE::SandboxedEditing
+{
+/** Holds all view models used by the root UI. Make sharing view models easier. */
+struct FBrowserViewModels
+{
+	/** The columns used by the browser widget. */
+	const TMap<FName, TArray<FSandboxColumnInfo>> BrowserColumns;
+	
+	/** Knows how to filter sandbox items. */
+	const TSharedRef<FFilterSandboxViewModel> FilterSandboxViewModel;
+	/** Knows about the items that are displayed in the browser. */
+	const TSharedRef<FSandboxListViewModel> ListViewModel;
+	
+	/** Handles actions that the user performs on the sandboxes via the UI. */
+	const TSharedRef<FSandboxControlsViewModel> ControlsViewModel;
+	/** Handles changing metadata of a sandbox. */
+	const TSharedRef<FSandboxMetaDataViewModel> MetaDataViewModel;
+	
+	/** Used for widgets to switch states when the active sandbox changes. */
+	const TSharedRef<FActiveSandboxTrackerViewModel> ActiveSandboxTrackerViewModel;
+	/** Used for widgets that display info about the currently active sandbox. */
+	const TSharedRef<FActiveSandboxDetailsViewModel> ActiveSandboxDetailsViewModel;
+	
+	explicit FBrowserViewModels(TMap<FName, TSharedRef<ISandboxColumnBehavior>> InColumns, const TSharedRef<FSandboxSystemModel>& InModel)
+		: FilterSandboxViewModel(MakeShared<FFilterSandboxViewModel>([&InColumns]
+			{
+				TArray<TSharedRef<ISandboxColumnBehavior>> Columns;
+				InColumns.GenerateValueArray(Columns);
+				return Columns;
+			}()))
+		, ListViewModel(MakeShared<FSandboxListViewModel>(InModel, FilterSandboxViewModel, MoveTemp(InColumns)))
+		, ControlsViewModel(MakeShared<FSandboxControlsViewModel>(InModel))
+		, MetaDataViewModel(MakeShared<FSandboxMetaDataViewModel>(InModel))
+		, ActiveSandboxTrackerViewModel(MakeShared<FActiveSandboxTrackerViewModel>(InModel))
+		, ActiveSandboxDetailsViewModel(MakeShared<FActiveSandboxDetailsViewModel>(InModel))
+	{}
+};
+}
diff --git a/Engine/Plugins/Developer/Sandbox/SandboxedEditing/Source/SandboxedEditing/Private/Features/Browser/ViewModels/Columns/ColumnInfo.h b/Engine/Plugins/Developer/Sandbox/SandboxedEditing/Source/SandboxedEditing/Private/Features/Browser/ViewModels/Columns/ColumnInfo.h
new file mode 100644
--- /dev/null
+++ b/Engine/Plugins/Developer/Sandbox/SandboxedEditing/Source/SandboxedEditing/Private/Features/Browser/ViewModels/Columns/ColumnInfo.h
@@ -0,0 +1,22 @@
+﻿// Copyright Epic Games, Inc. All Rights Reserved.
+
+#pragma once
+
+#include "IColumnBehavior.h"
+#include "IColumnWidgetFactory.h"
+#include "UObject/NameTypes.h"
+
+namespace UE::SandboxedEditing
+{
+template<typename TRowData>
+struct TColumnInfo
+{
+	TSharedRef<IColumnBehavior<TRowData>> Behavior;
+	TSharedRef<IColumnWidgetFactory<TRowData>> WidgetFactory;
+
+	explicit TColumnInfo(const TSharedRef<IColumnBehavior<TRowData>>& Behavior, const TSharedRef<IColumnWidgetFactory<TRowData>>& WidgetFactory
+		)
+		: Behavior(Behavior), WidgetFactory(WidgetFactory)
+	{}
+};
+}
diff --git a/Engine/Plugins/Developer/Sandbox/SandboxedEditing/Source/SandboxedEditing/Private/Features/Browser/ViewModels/Columns/IColumnBehavior.h b/Engine/Plugins/Developer/Sandbox/SandboxedEditing/Source/SandboxedEditing/Private/Features/Browser/ViewModels/Columns/IColumnBehavior.h
new file mode 100644
--- /dev/null
+++ b/Engine/Plugins/Developer/Sandbox/SandboxedEditing/Source/SandboxedEditing/Private/Features/Browser/ViewModels/Columns/IColumnBehavior.h
@@ -0,0 +1,20 @@
+﻿// Copyright Epic Games, Inc. All Rights Reserved.
+
+#pragma once
+
+#include "Widgets/Views/SHeaderRow.h"
+
+namespace UE::SandboxedEditing
+{
+/** API for the presentation layer to represent a layer, i.e. its view-model functionality, like sorting, search terms, etc. */
+template<typename TRowData>
+class IColumnBehavior
+{
+public:
+
+	/** Search terms for this column */
+	virtual void PopulateSearchTerms(const TRowData& InRowData, TArray<FString>& OutSearchTerms) const = 0;
+
+	virtual ~IColumnBehavior() = default;
+};
+}
diff --git a/Engine/Plugins/Developer/Sandbox/SandboxedEditing/Source/SandboxedEditing/Private/Features/Browser/ViewModels/Columns/IColumnWidgetFactory.h b/Engine/Plugins/Developer/Sandbox/SandboxedEditing/Source/SandboxedEditing/Private/Features/Browser/ViewModels/Columns/IColumnWidgetFactory.h
new file mode 100644
--- /dev/null
+++ b/Engine/Plugins/Developer/Sandbox/SandboxedEditing/Source/SandboxedEditing/Private/Features/Browser/ViewModels/Columns/IColumnWidgetFactory.h
@@ -0,0 +1,48 @@
+﻿// Copyright Epic Games, Inc. All Rights Reserved.
+
+#pragma once
+
+#include "Internationalization/Text.h"
+#include "Misc/Attribute.h"
+#include "UObject/NameTypes.h"
+#include "Widgets/Views/SHeaderRow.h"
+
+namespace UE::SandboxedEditing
+{
+/** Args for making a column widget in a row. */
+template<typename TRowData>
+struct TMakeColumnWidgetArgs
+{
+	/** The row data for this column widget */
+	const TRowData& RowData;
+		
+	/** Text that is being searched */
+	TAttribute<FText> HighlightText;
+	
+	/** Whether the row is currently selected in the SListView. Required for some widgets, e.g. SInlineEditableTextBlock. */
+	FIsSelected IsSelectedDelegate;
+
+	explicit TMakeColumnWidgetArgs(const TRowData& InRowData, TAttribute<FText> InHighlightText, FIsSelected InIsSelectedDelegate)
+		: RowData(InRowData)
+		, HighlightText(MoveTemp(InHighlightText))
+		, IsSelectedDelegate(MoveTemp(InIsSelectedDelegate))
+	{}
+};
+
+/** Interface for creating Slate widgets for a column. */
+template<typename TRowData, typename TMakeColumnArgs = TMakeColumnWidgetArgs<TRowData>>
+class IColumnWidgetFactory
+{
+public:
+	
+	using FMakeColumnWidgetArgs = TMakeColumnArgs;
+	
+	/** @return Info required to construct the column. */
+	virtual SHeaderRow::FColumn::FArguments MakeColumnArguments() = 0;
+	
+	/** @return Widget that should be put into the row for this column */
+	virtual TSharedRef<SWidget> MakeColumnWidget(const FMakeColumnWidgetArgs& InArgs) = 0;
+
+	virtual ~IColumnWidgetFactory() = default;
+};
+}
diff --git a/Engine/Plugins/Developer/Sandbox/SandboxedEditing/Source/SandboxedEditing/Private/Features/Browser/ViewModels/List/FilterSandboxViewModel.cpp b/Engine/Plugins/Developer/Sandbox/SandboxedEditing/Source/SandboxedEditing/Private/Features/Browser/ViewModels/List/FilterSandboxViewModel.cpp
new file mode 100644
--- /dev/null
+++ b/Engine/Plugins/Developer/Sandbox/SandboxedEditing/Source/SandboxedEditing/Private/Features/Browser/ViewModels/List/FilterSandboxViewModel.cpp
@@ -0,0 +1,41 @@
+﻿// Copyright Epic Games, Inc. All Rights Reserved.
+
+#include "FilterSandboxViewModel.h"
+
+#include "Framework/Models/SandboxInfo.h"
+
+namespace UE::SandboxedEditing
+{
+FFilterSandboxViewModel::FFilterSandboxViewModel(TAttribute<TArray<TSharedRef<ISandboxColumnBehavior>>> InVisibleColumns)
+	: VisibleColumns(MoveTemp(InVisibleColumns))
+	, TextFilter(FSandboxTextFilter::FItemToStringArray::CreateRaw(this, &FFilterSandboxViewModel::ItemToString))
+{
+	TextFilter.OnChanged().AddLambda([this]
+	{
+		OnFilterChangedDelegate.Broadcast();
+	});
+}
+
+FText FFilterSandboxViewModel::GetSearchText() const
+{
+	return TextFilter.GetRawFilterText();
+}
+
+void FFilterSandboxViewModel::SetSearchText(const FText& InSearchText)
+{
+	TextFilter.SetRawFilterText(InSearchText);
+}
+
+bool FFilterSandboxViewModel::PassesFilter(const TSharedRef<FSandboxListItem>& InItem) const
+{
+	return TextFilter.PassesFilter(InItem);
+}
+
+void FFilterSandboxViewModel::ItemToString(const TSharedRef<FSandboxListItem>& InItem, TArray<FString>& OutSearchTerms)
+{
+	for (const TSharedRef<ISandboxColumnBehavior>& Column : VisibleColumns.Get())
+	{
+		Column->PopulateSearchTerms(InItem, OutSearchTerms);
+	}
+}
+}
diff --git a/Engine/Plugins/Developer/Sandbox/SandboxedEditing/Source/SandboxedEditing/Private/Features/Browser/ViewModels/List/FilterSandboxViewModel.h b/Engine/Plugins/Developer/Sandbox/SandboxedEditing/Source/SandboxedEditing/Private/Features/Browser/ViewModels/List/FilterSandboxViewModel.h
new file mode 100644
--- /dev/null
+++ b/Engine/Plugins/Developer/Sandbox/SandboxedEditing/Source/SandboxedEditing/Private/Features/Browser/ViewModels/List/FilterSandboxViewModel.h
@@ -0,0 +1,48 @@
+﻿// Copyright Epic Games, Inc. All Rights Reserved.
+
+#pragma once
+
+#include "Features/Browser/Widgets/Unsandboxed/List/Columns/SandboxColumns.h"
+#include "Misc/Attribute.h"
+#include "Misc/TextFilter.h"
+#include "Templates/UnrealTemplate.h"
+
+namespace UE::SandboxedEditing
+{
+struct FSandboxInfo;
+
+/** Knows of all filters for filtering sandbox items. */
+class FFilterSandboxViewModel : public FNoncopyable
+{
+	using FSandboxTextFilter = TTextFilter<const TSharedRef<FSandboxListItem>&>;
+public:
+	
+	explicit FFilterSandboxViewModel(TAttribute<TArray<TSharedRef<ISandboxColumnBehavior>>> InVisibleColumns);
+	
+	/** @return Current search text */
+	FText GetSearchText() const;
+	/** Sets the search text. Will trigger a refilter. */
+	void SetSearchText(const FText& InSearchText);
+	
+	/** @return Whether the item can be displayed. */
+	bool PassesFilter(const TSharedRef<FSandboxListItem>& InItem) const;
+	
+	/** @return Delegate invoked when any of the filters changes */
+	FSimpleMulticastDelegate& OnFilterChanged() { return OnFilterChangedDelegate; }
+	
+private:
+	
+	/** Visible columns. Used for filtering visible columns by text. */
+	const TAttribute<TArray<TSharedRef<ISandboxColumnBehavior>>> VisibleColumns;
+	
+	/** Does text filltering on items */
+	FSandboxTextFilter TextFilter;
+	
+	/** Invoked when any of the filters changes */
+	FSimpleMulticastDelegate OnFilterChangedDelegate;
+	
+	/** Converts an item to string. */
+	void ItemToString(const TSharedRef<FSandboxListItem>& InItem, TArray<FString>& OutSearchTerms);
+};
+}
+
diff --git a/Engine/Plugins/Developer/Sandbox/SandboxedEditing/Source/SandboxedEditing/Private/Features/Browser/ViewModels/List/SandboxListItem.h b/Engine/Plugins/Developer/Sandbox/SandboxedEditing/Source/SandboxedEditing/Private/Features/Browser/ViewModels/List/SandboxListItem.h
new file mode 100644
--- /dev/null
+++ b/Engine/Plugins/Developer/Sandbox/SandboxedEditing/Source/SandboxedEditing/Private/Features/Browser/ViewModels/List/SandboxListItem.h
@@ -0,0 +1,22 @@
+﻿// Copyright Epic Games, Inc. All Rights Reserved.
+
+#pragma once
+
+#include "Framework/Models/SandboxInfo.h"
+
+namespace UE::SandboxedEditing
+{
+/** A single item displayed in the SSandboxListView. */
+class FSandboxListItem
+{
+public:
+	
+	/** Basic info about the sandbox */
+	FSandboxInfo SandboxInfo;
+
+	explicit FSandboxListItem(FSandboxInfo SandboxInfo)
+		: SandboxInfo(MoveTemp(SandboxInfo))
+	{}
+};
+}
+
diff --git a/Engine/Plugins/Developer/Sandbox/SandboxedEditing/Source/SandboxedEditing/Private/Features/Browser/ViewModels/List/SandboxListViewModel.cpp b/Engine/Plugins/Developer/Sandbox/SandboxedEditing/Source/SandboxedEditing/Private/Features/Browser/ViewModels/List/SandboxListViewModel.cpp
new file mode 100644
--- /dev/null
+++ b/Engine/Plugins/Developer/Sandbox/SandboxedEditing/Source/SandboxedEditing/Private/Features/Browser/ViewModels/List/SandboxListViewModel.cpp
@@ -0,0 +1,110 @@
+﻿// Copyright Epic Games, Inc. All Rights Reserved.
+
+#include "SandboxListViewModel.h"
+
+#include "FilterSandboxViewModel.h"
+#include "SandboxListItem.h"
+#include "Templates/SharedPointer.h"
+
+#define LOCTEXT_NAMESPACE "FSandboxListViewModel"
+
+namespace UE::SandboxedEditing
+{
+FSandboxListViewModel::FSandboxListViewModel(
+	const TSharedRef<FSandboxSystemModel>& InSandboxModel, const TSharedRef<FFilterSandboxViewModel>& InFilterViewModel,
+	TMap<FName, TSharedRef<ISandboxColumnBehavior>> InColumns
+	)
+	: SandboxModel(InSandboxModel)
+	, FilterViewModel(InFilterViewModel)
+	, Columns(MoveTemp(InColumns))
+{
+	RefreshItems();
+	
+	SandboxModel->OnKnownSandboxesChanged().AddRaw(this, &FSandboxListViewModel::RefreshItems);
+	FilterViewModel->OnFilterChanged().AddRaw(this, &FSandboxListViewModel::FilterItems);
+}
+
+FSandboxListViewModel::~FSandboxListViewModel()
+{
+	SandboxModel->OnKnownSandboxesChanged().RemoveAll(this);
+	FilterViewModel->OnFilterChanged().RemoveAll(this);
+}
+
+TArray<FName> FSandboxListViewModel::GetDisplayedColumns() const
+{
+	// For now all columns are visible. In the future, we may allow the user to hide columns.
+	TArray<FName> ColumnIds;
+	Columns.GenerateKeyArray(ColumnIds);
+	return ColumnIds;
+}
+
+EColumnSortPriority::Type FSandboxListViewModel::GetSortPriority(FName InColumnName) const
+{
+	// TODO DP:
+	return EColumnSortPriority::None;
+}
+
+EColumnSortMode::Type FSandboxListViewModel::GetSortMode(FName InColumnName) const
+{
+	// TODO DP:
+	return EColumnSortMode::None;
+}
+
+void FSandboxListViewModel::SetSortMode(const EColumnSortPriority::Type InSortPriority, const FName& ColumnId, const EColumnSortMode::Type InSortMode)
+{
+	// TODO DP:
+}
+
+void FSandboxListViewModel::RefreshItems()
+{
+	bool bMadeChanges = false;
+	TArray<TSharedPtr<FSandboxListItem>> OldItems = MoveTemp(AllUnfilteredItems);
+	SandboxModel->ForEachSandbox([this, &bMadeChanges, &OldItems](const FSandboxInfo& Info)
+	{
+		const TSharedPtr<FSandboxListItem>* ReusedItem = OldItems.FindByPredicate([&Info](const TSharedPtr<FSandboxListItem>& Item)
+		{
+			return FPaths::IsSamePath(Item->SandboxInfo.SandboxRoot, Info.SandboxRoot);
+		});
+		if (ReusedItem)
+		{
+			ReusedItem->Get()->SandboxInfo = Info; // Do this to make sure name and description are up to date.
+			AllUnfilteredItems.Add(*ReusedItem);
+		}
+		else
+		{
+			bMadeChanges = true;
+			AllUnfilteredItems.Emplace(MakeShared<FSandboxListItem>(Info));
+		}
+	});
+
+	bMadeChanges |= OldItems.Num() != AllUnfilteredItems.Num();
+	if (bMadeChanges)
+	{
+		FilterItems();
+	}
+}
+
+void FSandboxListViewModel::FilterItems()
+{
+	FilteredItems.Empty(AllUnfilteredItems.Num());
+	
+	for (const TSharedPtr<FSandboxListItem>& Item : AllUnfilteredItems)
+	{
+		if (ensure(Item) && FilterViewModel->PassesFilter(Item.ToSharedRef()))
+		{
+			FilteredItems.Add(Item);
+		}
+	}
+	
+	SortItems();
+	OnItemsChangedDelegate.Broadcast();
+}
+
+void FSandboxListViewModel::SortItems()
+{
+	// TODO DP: sort by sort mode & priority
+	AllUnfilteredItems.Sort();
+}
+}
+
+#undef LOCTEXT_NAMESPACE
\ No newline at end of file
diff --git a/Engine/Plugins/Developer/Sandbox/SandboxedEditing/Source/SandboxedEditing/Private/Features/Browser/ViewModels/List/SandboxListViewModel.h b/Engine/Plugins/Developer/Sandbox/SandboxedEditing/Source/SandboxedEditing/Private/Features/Browser/ViewModels/List/SandboxListViewModel.h
new file mode 100644
--- /dev/null
+++ b/Engine/Plugins/Developer/Sandbox/SandboxedEditing/Source/SandboxedEditing/Private/Features/Browser/ViewModels/List/SandboxListViewModel.h
@@ -0,0 +1,72 @@
+﻿// Copyright Epic Games, Inc. All Rights Reserved.
+
+#pragma once
+
+#include "Containers/Array.h"
+#include "Features/Browser/ViewModels/Columns/IColumnBehavior.h"
+#include "Features/Browser/Widgets/Unsandboxed/List/Columns/SandboxColumns.h"
+#include "Framework/Models/SandboxSystemModel.h"
+#include "Templates/UnrealTemplate.h"
+#include "Widgets/Views/SHeaderRow.h"
+
+namespace UE::SandboxedEditing
+{
+class FFilterSandboxViewModel;
+class FSandboxListItem;
+
+/** Decides what is displayed in SSandboxListView. */
+class FSandboxListViewModel : public FNoncopyable
+{
+public:
+	
+	explicit FSandboxListViewModel(
+		const TSharedRef<FSandboxSystemModel>& InSandboxModel, const TSharedRef<FFilterSandboxViewModel>& InFilterViewModel,
+		TMap<FName, TSharedRef<ISandboxColumnBehavior>> InColumns
+		);
+	~FSandboxListViewModel();
+	
+	/** @return The items to display */
+	const TArray<TSharedPtr<FSandboxListItem>>& GetItems() const { return FilteredItems; }
+	
+	/** @return The columns that should be displayed. */
+	TArray<FName> GetDisplayedColumns() const;
+	
+	/** @return The priority with which the column is sorted. */
+	EColumnSortPriority::Type GetSortPriority(FName InColumnName) const;
+	/** @return How the column is sorted */
+	EColumnSortMode::Type GetSortMode(FName InColumnName) const;
+	/** Changes the sort mode for the given column */
+	void SetSortMode(const EColumnSortPriority::Type InSortPriority, const FName& ColumnId, const EColumnSortMode::Type InSortMode);
+	
+	/** Invoked when the displayed items changes (added / removed / sorted).*/
+	FSimpleMulticastDelegate& OnItemsChanged() { return OnItemsChangedDelegate; }
+	
+private:
+	
+	/** The shared sandbox model. */
+	const TSharedRef<FSandboxSystemModel> SandboxModel;
+	/** Used to filter the items. */
+	const TSharedRef<FFilterSandboxViewModel> FilterViewModel;
+	
+	/** All columns that can ever be displayed. */
+	const TMap<FName, TSharedRef<ISandboxColumnBehavior>> Columns;
+	
+	/** Contains all items. */
+	TArray<TSharedPtr<FSandboxListItem>> AllUnfilteredItems;
+	/** A filtered version of AllUnfilteredItems. */
+	TArray<TSharedPtr<FSandboxListItem>> FilteredItems;
+	
+	/** Invoked when the displayed items changes (added / removed / sorted).*/
+	FSimpleMulticastDelegate OnItemsChangedDelegate;
+	
+	/** Updates AllUnfilteredItems from the model. Tries to retain preexisting items. */
+	void RefreshItems();
+	
+	/**  */
+	void FilterItems();
+	
+	/** Sorts the items. */
+	void SortItems();
+};
+}
+
diff --git a/Engine/Plugins/Developer/Sandbox/SandboxedEditing/Source/SandboxedEditing/Private/Features/Browser/ViewModels/SandboxControlsViewModel.cpp b/Engine/Plugins/Developer/Sandbox/SandboxedEditing/Source/SandboxedEditing/Private/Features/Browser/ViewModels/SandboxControlsViewModel.cpp
new file mode 100644
--- /dev/null
+++ b/Engine/Plugins/Developer/Sandbox/SandboxedEditing/Source/SandboxedEditing/Private/Features/Browser/ViewModels/SandboxControlsViewModel.cpp
@@ -0,0 +1,61 @@
+﻿// Copyright Epic Games, Inc. All Rights Reserved.
+
+#include "SandboxControlsViewModel.h"
+
+namespace UE::SandboxedEditing
+{
+FSandboxControlsViewModel::FSandboxControlsViewModel(const TSharedRef<FSandboxSystemModel>& InModel)
+	: Model(InModel)
+{}
+
+void FSandboxControlsViewModel::StartCreationWorkflow()
+{
+	if (CanStartCreationWorkflow())
+	{
+		CurrentCreationWorkflow = MakeUnique<FSandboxCreationWorkflow>(Model);
+		CurrentCreationWorkflow->OnWorkflowEnded().AddRaw(this, &FSandboxControlsViewModel::OnWorkflowEnded);
+		OnCreationWorkflowStartedDelegate.Broadcast(*CurrentCreationWorkflow);
+	}
+}
+
+bool FSandboxControlsViewModel::CanStartCreationWorkflow() const
+{
+	return CurrentCreationWorkflow == nullptr;
+}
+
+bool FSandboxControlsViewModel::CanLoadSandbox(const FString& InSandboxRoot, FText* OutReason) const
+{
+	return Model->CanLoadSandbox(InSandboxRoot);
+}
+
+void FSandboxControlsViewModel::LoadSandbox(const FString& InSandboxRoot)
+{
+	if (!CanLoadSandbox(InSandboxRoot))
+	{
+		return;
+	}
+	
+	if (CurrentCreationWorkflow)
+	{
+		CurrentCreationWorkflow->Cancel();
+	}
+	
+	Model->LoadSandbox(InSandboxRoot);
+}
+
+bool FSandboxControlsViewModel::CanDeleteSandbox(const FString& InSandboxRoot, FText* OutReason) const
+{
+	return Model->CanDeleteSandbox(InSandboxRoot, OutReason);
+}
+
+void FSandboxControlsViewModel::DeleteSandbox(const FString& InSandboxRoot) const
+{
+	Model->DeleteSandbox(InSandboxRoot);
+}
+
+void FSandboxControlsViewModel::OnWorkflowEnded()
+{
+	CurrentCreationWorkflow.Reset();
+	OnCreationWorkflowEnded().Broadcast();
+}
+}
diff --git a/Engine/Plugins/Developer/Sandbox/SandboxedEditing/Source/SandboxedEditing/Private/Features/Browser/ViewModels/SandboxControlsViewModel.h b/Engine/Plugins/Developer/Sandbox/SandboxedEditing/Source/SandboxedEditing/Private/Features/Browser/ViewModels/SandboxControlsViewModel.h
new file mode 100644
--- /dev/null
+++ b/Engine/Plugins/Developer/Sandbox/SandboxedEditing/Source/SandboxedEditing/Private/Features/Browser/ViewModels/SandboxControlsViewModel.h
@@ -0,0 +1,66 @@
+﻿// Copyright Epic Games, Inc. All Rights Reserved.
+
+#pragma once
+
+#include "Delegates/Delegate.h"
+#include "Framework/Models/SandboxSystemModel.h"
+#include "SandboxCreationWorkflow.h"
+#include "Templates/UniquePtr.h"
+#include "Templates/UnrealTemplate.h"
+
+namespace UE::SandboxedEditing
+{
+DECLARE_MULTICAST_DELEGATE_OneParam(FCreationWorkflowDelegate, const FSandboxCreationWorkflow&);
+
+/**
+ * Handles actions that can be performed with FSandboxListItems, such as deleting, renaming, loading.
+ * 
+ * These actions are primarily present on the controls widget in the form of buttons.
+ * Some of them can be performed in the list view, such as deleting or renaming.
+ */
+class FSandboxControlsViewModel : public FNoncopyable
+{
+public:
+	
+	explicit FSandboxControlsViewModel(const TSharedRef<FSandboxSystemModel>& InModel);
+	
+	/** Starts a workflow for creating a new sandbox. */
+	void StartCreationWorkflow();
+	/** @return Whether a new sandbox creation workflow can currently be started. */
+	bool CanStartCreationWorkflow() const;
+	/** @return The currently active creation workflow. */
+	FSandboxCreationWorkflow* GetCurrentCreationWorkflow() const { return CurrentCreationWorkflow.Get(); }
+
+	/** @return Whether the sandbox can be loaded. */
+	bool CanLoadSandbox(const FString& InSandboxRoot, FText* OutReason = nullptr) const;
+	/** Loads the sandbox. */
+	void LoadSandbox(const FString& InSandboxRoot);
+	
+	/** @return Whether the sandbox can be deleted. */
+	bool CanDeleteSandbox(const FString& InSandboxRoot, FText* OutReason = nullptr) const;
+	/** Deletes the sandbox. */
+	void DeleteSandbox(const FString& InSandboxRoot) const;
+	
+	/** @return Delegate that is invoked when the creation workflow starts. */
+	FCreationWorkflowDelegate& OnCreationWorkflowStarted() { return OnCreationWorkflowStartedDelegate; }
+	/** @return Delegate that is invoked when the creation workflow ends. */
+	FSimpleMulticastDelegate& OnCreationWorkflowEnded() { return OnCreationWorkflowEndedDelegate; }
+	
+private:
+	
+	/** The underlying model that performs the mutations. */
+	const TSharedRef<FSandboxSystemModel> Model;
+	
+	/** Set if the user is currently in a process of creating a sandbox. */
+	TUniquePtr<FSandboxCreationWorkflow> CurrentCreationWorkflow;
+	
+	/** Invoked when the creation workflow starts. */
+	FCreationWorkflowDelegate OnCreationWorkflowStartedDelegate;
+	/** Invoked when the creation workflow ends. */
+	FSimpleMulticastDelegate OnCreationWorkflowEndedDelegate;
+	
+	/** Called by CurrentCreationWorkflow when the workflow ends. */
+	void OnWorkflowEnded();
+};
+}
+
diff --git a/Engine/Plugins/Developer/Sandbox/SandboxedEditing/Source/SandboxedEditing/Private/Features/Browser/ViewModels/SandboxCreationWorkflow.cpp b/Engine/Plugins/Developer/Sandbox/SandboxedEditing/Source/SandboxedEditing/Private/Features/Browser/ViewModels/SandboxCreationWorkflow.cpp
new file mode 100644
--- /dev/null
+++ b/Engine/Plugins/Developer/Sandbox/SandboxedEditing/Source/SandboxedEditing/Private/Features/Browser/ViewModels/SandboxCreationWorkflow.cpp
@@ -0,0 +1,53 @@
+﻿// Copyright Epic Games, Inc. All Rights Reserved.
+
+#include "SandboxCreationWorkflow.h"
+
+#include "Framework/Models/CreateSandboxArgs.h"
+#include "Framework/Models/SandboxSystemModel.h"
+
+namespace UE::SandboxedEditing
+{
+FSandboxCreationWorkflow::FSandboxCreationWorkflow(const TSharedRef<FSandboxSystemModel>& InModel)
+	: Model(InModel)
+{}
+
+void FSandboxCreationWorkflow::Confirm()
+{
+	if (!CanConfirm())
+	{
+		return;
+	}
+	
+	if (!Model->CreateNewSandbox(FCreateSandboxArgs(Name, Description)))
+	{
+		return;
+	}
+	
+	OnWorkflowEndedDelegate.Broadcast();
+}
+
+void FSandboxCreationWorkflow::Cancel()
+{
+	OnWorkflowEndedDelegate.Broadcast();
+}
+
+bool FSandboxCreationWorkflow::CanConfirm(FText* OutReason) const
+{
+	return IsValidName(Name, OutReason);
+}
+
+void FSandboxCreationWorkflow::SetName(const FString& InNewName)
+{
+	Name = InNewName;
+}
+
+bool FSandboxCreationWorkflow::IsValidName(const FString& InName, FText* OutReason) const
+{
+	return Model->CanCreateNewSandbox(InName, OutReason);
+}
+
+void FSandboxCreationWorkflow::SetDescription(const FString& InNewDescription)
+{
+	Description = InNewDescription;
+}
+}
diff --git a/Engine/Plugins/Developer/Sandbox/SandboxedEditing/Source/SandboxedEditing/Private/Features/Browser/ViewModels/SandboxCreationWorkflow.h b/Engine/Plugins/Developer/Sandbox/SandboxedEditing/Source/SandboxedEditing/Private/Features/Browser/ViewModels/SandboxCreationWorkflow.h
new file mode 100644
--- /dev/null
+++ b/Engine/Plugins/Developer/Sandbox/SandboxedEditing/Source/SandboxedEditing/Private/Features/Browser/ViewModels/SandboxCreationWorkflow.h
@@ -0,0 +1,57 @@
+﻿// Copyright Epic Games, Inc. All Rights Reserved.
+
+#pragma once
+
+#include "Delegates/Delegate.h"
+#include "Templates/UnrealTemplate.h"
+
+namespace UE::SandboxedEditing
+{
+class FSandboxSystemModel;
+
+/** 
+ * View-model managing the workflow for creating a new sandbox.
+ * 
+ * The user can type the name, edit other data, etc. 
+ * At any point they can confirm or cancel the workflow. 
+ */
+class FSandboxCreationWorkflow : public FNoncopyable
+{
+public:
+	
+	explicit FSandboxCreationWorkflow(const TSharedRef<FSandboxSystemModel>& InModel);
+	
+	/** Finishes the workflow, if possible, resulting in the creation of the sandbox. */
+	void Confirm();
+	/** Stops the workflow. */
+	void Cancel();
+	
+	/** @return Whether the workflow can be confirmed, i.e. whether the sandbox can be created. */
+	bool CanConfirm(FText* OutReason = nullptr) const;
+	
+	/** Sets the name the user wants to give the sandbox. The name needn't be valid, however, it will prevent CanConfirm from returning true. */
+	void SetName(const FString& InNewName);
+	/** @return Whether this name is a valid name to set. */
+	bool IsValidName(const FString& InName, FText* OutReason = nullptr) const;
+	
+	/** Sets the description. There are no restrictions on the description. */
+	void SetDescription(const FString& InNewDescription);
+	
+	/** Invoked when the workflow ends. */
+	FSimpleMulticastDelegate& OnWorkflowEnded() { return OnWorkflowEndedDelegate; }
+	
+private:  
+	
+	/** The model used to interact with the sandbox system. */
+	const TSharedRef<FSandboxSystemModel> Model;
+	
+	/** Invoked when the workflow ends. */
+	FSimpleMulticastDelegate OnWorkflowEndedDelegate;
+	
+	/** The name the user wants to give. */
+	FString Name;
+	/** The description the user wants to give. */
+	FString Description;
+};
+}
+
diff --git a/Engine/Plugins/Developer/Sandbox/SandboxedEditing/Source/SandboxedEditing/Private/Features/Browser/ViewModels/SandboxMetaDataViewModel.cpp b/Engine/Plugins/Developer/Sandbox/SandboxedEditing/Source/SandboxedEditing/Private/Features/Browser/ViewModels/SandboxMetaDataViewModel.cpp
new file mode 100644
--- /dev/null
+++ b/Engine/Plugins/Developer/Sandbox/SandboxedEditing/Source/SandboxedEditing/Private/Features/Browser/ViewModels/SandboxMetaDataViewModel.cpp
@@ -0,0 +1,24 @@
+﻿// Copyright Epic Games, Inc. All Rights Reserved.
+
+#include "SandboxMetaDataViewModel.h"
+
+#include "Framework/Models/SandboxInfo.h"
+#include "Framework/Models/SandboxSystemModel.h"
+
+namespace UE::SandboxedEditing
+{
+FSandboxMetaDataViewModel::FSandboxMetaDataViewModel(const TSharedRef<FSandboxSystemModel>& InModel)
+	: Model(InModel)
+{}
+
+TOptional<FString> FSandboxMetaDataViewModel::GetDescription(const FString& InSandboxRoot) const
+{
+	const TOptional<FSandboxInfo> Info = Model->GetSandboxInfo(InSandboxRoot);
+	return Info ? Info->Description : TOptional<FString>();
+}
+
+void FSandboxMetaDataViewModel::SetDescription(const FString& InSandboxRoot, const FString& InNewDescription)
+{
+	Model->SetDescription(InSandboxRoot, InNewDescription);
+}
+}
diff --git a/Engine/Plugins/Developer/Sandbox/SandboxedEditing/Source/SandboxedEditing/Private/Features/Browser/ViewModels/SandboxMetaDataViewModel.h b/Engine/Plugins/Developer/Sandbox/SandboxedEditing/Source/SandboxedEditing/Private/Features/Browser/ViewModels/SandboxMetaDataViewModel.h
new file mode 100644
--- /dev/null
+++ b/Engine/Plugins/Developer/Sandbox/SandboxedEditing/Source/SandboxedEditing/Private/Features/Browser/ViewModels/SandboxMetaDataViewModel.h
@@ -0,0 +1,31 @@
+﻿// Copyright Epic Games, Inc. All Rights Reserved.
+
+#pragma once
+
+#include "Templates/SharedPointer.h"
+#include "Templates/UnrealTemplate.h"
+
+namespace UE::SandboxedEditing
+{
+class FSandboxSystemModel;
+
+/** Handles editing metadata about sandboxes, such as name and description. */
+class FSandboxMetaDataViewModel : public FNoncopyable
+{
+public:
+	
+	explicit FSandboxMetaDataViewModel(const TSharedRef<FSandboxSystemModel>& InModel);
+	
+	/** @return The description of the given sandbox. Unset if the sandbox root does not exist. */
+	TOptional<FString> GetDescription(const FString& InSandboxRoot) const;
+	
+	/** Sets the description of the sandbox */
+	void SetDescription(const FString& InSandboxRoot, const FString& InNewDescription);
+	
+private:
+	
+	/** The model used to query the sandboxes. */
+	const TSharedRef<FSandboxSystemModel> Model;
+};
+}
+
diff --git a/Engine/Plugins/Developer/Sandbox/SandboxedEditing/Source/SandboxedEditing/Private/Features/Browser/Widgets/SSandboxedEditingRoot.cpp b/Engine/Plugins/Developer/Sandbox/SandboxedEditing/Source/SandboxedEditing/Private/Features/Browser/Widgets/SSandboxedEditingRoot.cpp
new file mode 100644
--- /dev/null
+++ b/Engine/Plugins/Developer/Sandbox/SandboxedEditing/Source/SandboxedEditing/Private/Features/Browser/Widgets/SSandboxedEditingRoot.cpp
@@ -0,0 +1,72 @@
+﻿// Copyright Epic Games, Inc. All Rights Reserved.
+
+#include "SSandboxedEditingRoot.h"
+
+#include "Features/Browser/Commands/BrowserCommandBindings.h"
+#include "Features/Browser/ViewModels/BrowserViewModels.h"
+#include "Features/Browser/ViewModels/Active/ActiveSandboxTrackerViewModel.h"
+#include "Features/Browser/Widgets/Unsandboxed/SSandboxedEditingBrowser.h"
+#include "Sandboxed/SActiveSandbox.h"
+#include "Widgets/Layout/SWidgetSwitcher.h"
+
+namespace UE::SandboxedEditing
+{
+void SSandboxedEditingRoot::Construct(const FArguments& InArgs, const FBrowserViewModels& InViewModels)
+{
+	CommandBindings = InArgs._CommandList;
+	ActiveSandboxTrackerViewModel = InViewModels.ActiveSandboxTrackerViewModel;
+	ActiveSandboxViewModels.Emplace(InViewModels.ActiveSandboxDetailsViewModel, InViewModels.MetaDataViewModel);
+	
+	ActiveSandboxTrackerViewModel->OnLoadSandbox().AddSP(this, &SSandboxedEditingRoot::ShowActiveSandbox);
+	ActiveSandboxTrackerViewModel->OnLeaveSandbox().AddSP(this, &SSandboxedEditingRoot::ShowBrowser);
+	
+	const SSandboxedEditingBrowser::FViewModels ViewModels 
+	{ 
+		InViewModels.FilterSandboxViewModel, InViewModels.ListViewModel, InViewModels.ControlsViewModel, InViewModels.MetaDataViewModel 
+	};
+	Browser = SNew(SSandboxedEditingBrowser, ViewModels)
+		.CommandList(CommandBindings)
+		.ColumnFactories(InArgs._ColumnFactories);
+	
+	UpdateWidgetContent();
+}
+
+FReply SSandboxedEditingRoot::OnKeyDown(const FGeometry& MyGeometry, const FKeyEvent& InKeyEvent)
+{
+	if (CommandBindings && CommandBindings->ProcessCommandBindings(InKeyEvent))
+	{
+		return FReply::Handled();
+	}
+	
+	return SCompoundWidget::OnKeyDown(MyGeometry, InKeyEvent);
+}
+
+void SSandboxedEditingRoot::UpdateWidgetContent()
+{
+	if (ActiveSandboxTrackerViewModel->HasActiveSandbox())
+	{
+		ShowActiveSandbox();
+	}
+	else
+	{
+		ShowBrowser();
+	}
+}
+
+void SSandboxedEditingRoot::ShowActiveSandbox()
+{
+	ChildSlot
+	[
+		SNew(SActiveSandbox, *ActiveSandboxViewModels)
+		.CommandList(CommandBindings)
+	];
+}
+
+void SSandboxedEditingRoot::ShowBrowser()
+{
+	ChildSlot
+	[
+		Browser.ToSharedRef()
+	];
+}
+}
diff --git a/Engine/Plugins/Developer/Sandbox/SandboxedEditing/Source/SandboxedEditing/Private/Features/Browser/Widgets/SSandboxedEditingRoot.h b/Engine/Plugins/Developer/Sandbox/SandboxedEditing/Source/SandboxedEditing/Private/Features/Browser/Widgets/SSandboxedEditingRoot.h
new file mode 100644
--- /dev/null
+++ b/Engine/Plugins/Developer/Sandbox/SandboxedEditing/Source/SandboxedEditing/Private/Features/Browser/Widgets/SSandboxedEditingRoot.h
@@ -0,0 +1,61 @@
+﻿// Copyright Epic Games, Inc. All Rights Reserved.
+
+#pragma once
+
+#include "Features/Browser/ViewModels/Active/ActiveSandboxDetailsViewModel.h"
+#include "Sandboxed/SActiveSandbox.h"
+#include "Unsandboxed/List/Columns/SandboxColumns.h"
+#include "Widgets/DeclarativeSyntaxSupport.h"
+#include "Widgets/SCompoundWidget.h"
+
+class FUICommandList;
+
+namespace UE::SandboxedEditing
+{
+class FSandboxMetaDataViewModel;
+class FActiveSandboxTrackerViewModel;
+class FBrowserCommandBindings;
+class FSandboxControlsViewModel;
+class FSandboxListViewModel;
+class SSandboxedEditingBrowser;
+struct FBrowserViewModels;
+
+/** Root widget for Sandboxed Editing. It updates it content depending on whether the engine is sandboxed. */
+class SSandboxedEditingRoot : public SCompoundWidget
+{
+public:
+
+	SLATE_BEGIN_ARGS(SSandboxedEditingRoot){}
+		/** The command list to bind selection-based commands to. */
+		SLATE_ARGUMENT(TSharedPtr<FUICommandList>, CommandList)
+		/** Used to build the columns. */
+		SLATE_ARGUMENT(FSandboxColumnFactoryMap, ColumnFactories)
+	SLATE_END_ARGS()
+
+	void Construct(const FArguments& InArgs, const FBrowserViewModels& InViewModels);
+	
+	//~ Begin SWidget Interface
+	virtual FReply OnKeyDown(const FGeometry& MyGeometry, const FKeyEvent& InKeyEvent) override;
+	//~ End SWidget Interface
+	
+private:
+	
+	/** The sandbox editing browser. Shown as long as not in any sandbox. */
+	TSharedPtr<SSandboxedEditingBrowser> Browser;
+	
+	/** Knows when the active sandbox changes. */
+	TSharedPtr<FActiveSandboxTrackerViewModel> ActiveSandboxTrackerViewModel;
+	/** The view models required to construct the active sandbox widget. Set once Construct has been called. */
+	TOptional<SActiveSandbox::FViewModels> ActiveSandboxViewModels;
+	
+	/** Command bindings for the browser UI. */
+	TSharedPtr<FUICommandList> CommandBindings;
+	
+	/** Sets the widget content to the browser or active session widget depending on whether we're in a sandbox. */
+	void UpdateWidgetContent();
+	
+	void ShowActiveSandbox();
+	void ShowBrowser();
+};
+}
+
diff --git a/Engine/Plugins/Developer/Sandbox/SandboxedEditing/Source/SandboxedEditing/Private/Features/Browser/Widgets/Sandboxed/SActiveSandbox.cpp b/Engine/Plugins/Developer/Sandbox/SandboxedEditing/Source/SandboxedEditing/Private/Features/Browser/Widgets/Sandboxed/SActiveSandbox.cpp
new file mode 100644
--- /dev/null
+++ b/Engine/Plugins/Developer/Sandbox/SandboxedEditing/Source/SandboxedEditing/Private/Features/Browser/Widgets/Sandboxed/SActiveSandbox.cpp
@@ -0,0 +1,78 @@
+﻿// Copyright Epic Games, Inc. All Rights Reserved.
+
+#include "SActiveSandbox.h"
+
+#include "Features/Browser/ViewModels/Active/ActiveSandboxDetailsViewModel.h"
+#include "Features/Browser/Widgets/Shared/SSandboxDescription.h"
+#include "SActiveSandboxToolbar.h"
+#include "Styling/ToolBarStyle.h"
+#include "Widgets/Layout/SBorder.h"
+#include "Widgets/Layout/SBox.h"
+#include "Widgets/Layout/SExpandableArea.h"
+#include "Widgets/Layout/SSeparator.h"
+#include "Widgets/SBoxPanel.h"
+#include "Widgets/Text/STextBlock.h"
+
+#define LOCTEXT_NAMESPACE "SActiveSandbox"
+
+namespace UE::SandboxedEditing
+{
+void SActiveSandbox::Construct(const FArguments& InArgs, const FViewModels& InViewModels)
+{
+	ActiveSandboxViewModel = InViewModels.ActiveSandboxViewModel;
+	
+	ChildSlot
+	[
+		SNew(SBorder)
+		.BorderImage(FAppStyle::Get().GetBrush("ToolPanel.GroupBorder"))
+		[
+			SNew(SVerticalBox)
+
+			+SVerticalBox::Slot()
+			.AutoHeight()
+			.Padding(1)
+			[
+				SNew(SActiveSandboxToolbar, InViewModels.ActiveSandboxViewModel)
+					.CommandList(InArgs._CommandList)
+			]
+
+			+SVerticalBox::Slot()
+			.AutoHeight()
+			[
+				SNew(SSeparator)
+				.Thickness(FCoreStyle::Get().GetWidgetStyle<FToolBarStyle>("ToolBar").SeparatorThickness)
+				.SeparatorImage(&FCoreStyle::Get().GetWidgetStyle<FToolBarStyle>("ToolBar").SeparatorBrush)
+			]
+
+			+SVerticalBox::Slot()
+			.AutoHeight()
+			.Padding(1)
+			[
+				MakeDescription(InViewModels)
+			]
+		]
+	];
+}
+
+TSharedRef<SWidget> SActiveSandbox::MakeDescription(const FViewModels& InViewModels)
+{
+	return SNew(SExpandableArea)
+		.BodyBorderImage(FAppStyle::Get().GetBrush("ToolPanel.GroupBorder"))
+		.HeaderContent()
+		[
+			SNew(STextBlock)
+			.Text(LOCTEXT("Description.Label", "Description"))
+		]
+		.BodyContent()
+		[
+			SNew(SBox)
+			.MinDesiredHeight(120)
+			[
+				SNew(SSandboxDescription, InViewModels.MetaDataViewModel)
+				.SandboxPath(ActiveSandboxViewModel->GetSandboxPath())
+			]
+		];
+}
+}
+
+#undef LOCTEXT_NAMESPACE
\ No newline at end of file
diff --git a/Engine/Plugins/Developer/Sandbox/SandboxedEditing/Source/SandboxedEditing/Private/Features/Browser/Widgets/Sandboxed/SActiveSandbox.h b/Engine/Plugins/Developer/Sandbox/SandboxedEditing/Source/SandboxedEditing/Private/Features/Browser/Widgets/Sandboxed/SActiveSandbox.h
new file mode 100644
--- /dev/null
+++ b/Engine/Plugins/Developer/Sandbox/SandboxedEditing/Source/SandboxedEditing/Private/Features/Browser/Widgets/Sandboxed/SActiveSandbox.h
@@ -0,0 +1,50 @@
+﻿// Copyright Epic Games, Inc. All Rights Reserved.
+
+#pragma once
+
+#include "Widgets/DeclarativeSyntaxSupport.h"
+#include "Widgets/SCompoundWidget.h"
+
+class FUICommandList;
+
+namespace UE::SandboxedEditing
+{
+class FActiveSandboxDetailsViewModel;
+class FSandboxMetaDataViewModel;
+
+/** Widget displayed while in a sandbox */
+class SActiveSandbox : public SCompoundWidget
+{
+public:
+	
+	struct FViewModels
+	{
+		TSharedRef<FActiveSandboxDetailsViewModel> ActiveSandboxViewModel;
+		TSharedRef<FSandboxMetaDataViewModel> MetaDataViewModel;
+
+		explicit FViewModels(
+			const TSharedRef<FActiveSandboxDetailsViewModel>& InActiveSandboxViewModel,
+			const TSharedRef<FSandboxMetaDataViewModel>& InMetaDataViewModel
+			)
+			: ActiveSandboxViewModel(InActiveSandboxViewModel)
+			, MetaDataViewModel(InMetaDataViewModel)
+		{}
+	};
+	
+	SLATE_BEGIN_ARGS(SActiveSandbox) {}
+		/** The command list to bind commands to. */
+		SLATE_ARGUMENT(TSharedPtr<FUICommandList>, CommandList)
+	SLATE_END_ARGS()
+
+	void Construct(const FArguments& InArgs, const FViewModels& InViewModels);
+	
+private:
+	
+	/** Details about the active sandbox. */
+	TSharedPtr<FActiveSandboxDetailsViewModel> ActiveSandboxViewModel;
+	
+	/** Make the description widget. */
+	TSharedRef<SWidget> MakeDescription(const FViewModels& InViewModels);
+};
+}
+
diff --git a/Engine/Plugins/Developer/Sandbox/SandboxedEditing/Source/SandboxedEditing/Private/Features/Browser/Widgets/Sandboxed/SActiveSandboxToolbar.cpp b/Engine/Plugins/Developer/Sandbox/SandboxedEditing/Source/SandboxedEditing/Private/Features/Browser/Widgets/Sandboxed/SActiveSandboxToolbar.cpp
new file mode 100644
--- /dev/null
+++ b/Engine/Plugins/Developer/Sandbox/SandboxedEditing/Source/SandboxedEditing/Private/Features/Browser/Widgets/Sandboxed/SActiveSandboxToolbar.cpp
@@ -0,0 +1,75 @@
+﻿// Copyright Epic Games, Inc. All Rights Reserved.
+
+#include "SActiveSandboxToolbar.h"
+
+#include "Features/Browser/Commands/BrowserCommands.h"
+#include "Features/Browser/ViewModels/Active/ActiveSandboxDetailsViewModel.h"
+#include "SandboxedEditingStyle.h"
+#include "Styling/AppStyle.h"
+#include "Utils/WidgetUtils.h"
+#include "Widgets/Layout/SBorder.h"
+#include "Widgets/Layout/SSeparator.h"
+#include "Widgets/Layout/SSpacer.h"
+#include "Widgets/SBoxPanel.h"
+#include "Widgets/Text/STextBlock.h"
+
+namespace UE::SandboxedEditing
+{
+void SActiveSandboxToolbar::Construct(const FArguments& InArgs, const TSharedRef<FActiveSandboxDetailsViewModel>& InActiveSandboxViewModel)
+{
+	ActiveSandboxViewModel = InActiveSandboxViewModel;
+	
+	ChildSlot
+	[
+		SNew(SHorizontalBox)
+
+		+SHorizontalBox::Slot()
+		.AutoWidth()
+		.VAlign(VAlign_Center)
+		.Padding(FMargin(4.0f, 1.0f))
+		[
+			MakeNameWidget()
+		]
+
+		+SHorizontalBox::Slot()
+		.AutoWidth()
+		.Padding(4, 3)
+		[
+			SNew(SSeparator)
+			.Orientation(Orient_Vertical)
+		]
+
+		// Move all the other content to the right if the panel is big
+		+SHorizontalBox::Slot()
+		.FillWidth(1.f)
+		[
+			SNew(SSpacer)
+		]
+
+		// Leave button
+		+SHorizontalBox::Slot()
+		.AutoWidth()
+		.VAlign(VAlign_Center)
+		.Padding(FMargin(1.0f))
+		[
+			MakeIconButtonByCommand(
+				FSandboxedEditingStyle::Get().GetBrush("SandboxedEditing.Browser.LeaveSandbox"),
+				InArgs._CommandList, FBrowserCommands::Get().LeaveSandbox
+				)
+		]
+	];
+}
+
+TSharedRef<SWidget> SActiveSandboxToolbar::MakeNameWidget() const
+{
+	return SNew(SBorder)
+		.BorderImage(FAppStyle::Get().GetBrush("NoBorder"))
+		.ColorAndOpacity(FLinearColor(0.75f, 0.75f, 0.75f))
+		.Padding(FMargin(0.0f, 4.0f, 6.0f, 4.0f))
+		[
+			SNew(STextBlock)
+			.Font(FAppStyle::Get().GetFontStyle("BoldFont"))
+			.Text_Lambda([this]{ return FText::AsCultureInvariant(ActiveSandboxViewModel->GetSandboxName()); })
+		];
+}
+}
diff --git a/Engine/Plugins/Developer/Sandbox/SandboxedEditing/Source/SandboxedEditing/Private/Features/Browser/Widgets/Sandboxed/SActiveSandboxToolbar.h b/Engine/Plugins/Developer/Sandbox/SandboxedEditing/Source/SandboxedEditing/Private/Features/Browser/Widgets/Sandboxed/SActiveSandboxToolbar.h
new file mode 100644
--- /dev/null
+++ b/Engine/Plugins/Developer/Sandbox/SandboxedEditing/Source/SandboxedEditing/Private/Features/Browser/Widgets/Sandboxed/SActiveSandboxToolbar.h
@@ -0,0 +1,35 @@
+﻿// Copyright Epic Games, Inc. All Rights Reserved.
+
+#pragma once
+
+#include "Widgets/DeclarativeSyntaxSupport.h"
+#include "Widgets/SCompoundWidget.h"
+
+class FUICommandList;
+
+namespace UE::SandboxedEditing
+{
+class FActiveSandboxDetailsViewModel;
+
+/** Widget displayed on the top while in an active sandbox. */
+class SActiveSandboxToolbar : public SCompoundWidget
+{
+public:
+	
+	SLATE_BEGIN_ARGS(SActiveSandboxToolbar) {}
+		/** The command list to bind commands to. */
+		SLATE_ARGUMENT(TSharedPtr<FUICommandList>, CommandList)
+	SLATE_END_ARGS()
+	
+	void Construct(const FArguments& InArgs, const TSharedRef<FActiveSandboxDetailsViewModel>& InActiveSandboxViewModel);
+	
+private:
+	
+	/** Details about the active sandbox. */
+	TSharedPtr<FActiveSandboxDetailsViewModel> ActiveSandboxViewModel;
+	
+	/** @return Widget displaying the active sandbox name. */
+	TSharedRef<SWidget> MakeNameWidget() const;
+};
+}
+
diff --git a/Engine/Plugins/Developer/Sandbox/SandboxedEditing/Source/SandboxedEditing/Private/Features/Browser/Widgets/Shared/SSandboxDescription.cpp b/Engine/Plugins/Developer/Sandbox/SandboxedEditing/Source/SandboxedEditing/Private/Features/Browser/Widgets/Shared/SSandboxDescription.cpp
new file mode 100644
--- /dev/null
+++ b/Engine/Plugins/Developer/Sandbox/SandboxedEditing/Source/SandboxedEditing/Private/Features/Browser/Widgets/Shared/SSandboxDescription.cpp
@@ -0,0 +1,75 @@
+﻿// Copyright Epic Games, Inc. All Rights Reserved.
+
+#include "SSandboxDescription.h"
+
+#include "Features/Browser/ViewModels/SandboxMetaDataViewModel.h"
+#include "Framework/Notifications/NotificationManager.h"
+#include "SandboxedEditingStyle.h"
+#include "Widgets/Input/SMultiLineEditableTextBox.h"
+#include "Widgets/Layout/SBox.h"
+#include "Widgets/Notifications/SNotificationList.h"
+#include "Widgets/Text/STextBlock.h"
+
+#define LOCTEXT_NAMESPACE "SSandboxDescription"
+
+namespace UE::SandboxedEditing
+{
+void SSandboxDescription::Construct(const FArguments& InArgs, const TSharedRef<FSandboxMetaDataViewModel>& InMetaDataViewModel)
+{
+	MetaDataViewModel = InMetaDataViewModel;
+	SandboxPathAttr = InArgs._SandboxPath;
+	
+	ChildSlot
+	[
+		SAssignNew(DescriptionTextCtrl, SMultiLineEditableTextBox)
+		.SelectAllTextWhenFocused(true)
+		.ModiferKeyForNewLine(EModifierKey::Shift)
+		.AutoWrapText(true)
+		.RevertTextOnEscape(true)
+		.HintText(LOCTEXT("Description.Hint", "Enter a description..."))
+		.Text(this, &SSandboxDescription::GetDescription)
+		.OnTextCommitted(this, &SSandboxDescription::OnDescriptionCommitted)
+	];
+}
+
+FText SSandboxDescription::GetDescription() const
+{
+	const TOptional<FString> SandboxPath = SandboxPathAttr.Get();
+	if (!SandboxPath)
+	{
+		return FText::GetEmpty();
+	}
+	
+	const TOptional<FString> Description = MetaDataViewModel->GetDescription(*SandboxPath);
+	return Description ? FText::AsCultureInvariant(*Description) : FText::GetEmpty();
+}
+
+void SSandboxDescription::OnDescriptionCommitted(const FText& InNewDescription, ETextCommit::Type InType) const
+{
+	const TOptional<FString> SandboxPath = SandboxPathAttr.Get();
+	if (!SandboxPath)
+	{
+		return;
+	}
+	
+	const FString NewDescription = InNewDescription.ToString();
+	const TOptional<FString> CurrentDescription = MetaDataViewModel->GetDescription(*SandboxPath);
+	if (!CurrentDescription || NewDescription == *CurrentDescription)
+	{
+		return;
+	}
+	
+	MetaDataViewModel->SetDescription(*SandboxPath, InNewDescription.ToString());
+	
+	// Showing a notification gives feedback to the user that they caused something to happen. 
+	// Commit happens e.g. when you click outside the text edit box but would not "feel" like it actualy did anything.
+	FSlateNotificationManager& NotificationManager = FSlateNotificationManager::Get();
+	FNotificationInfo Info(LOCTEXT("DescriptionEdited", "Description updated"));
+	Info.bFireAndForget = true;
+	Info.ExpireDuration = 4.f;
+	const TSharedPtr<SNotificationItem> Notification = NotificationManager.AddNotification(Info);
+	Notification->SetCompletionState(SNotificationItem::CS_Success);
+}
+}
+
+#undef LOCTEXT_NAMESPACE
\ No newline at end of file
diff --git a/Engine/Plugins/Developer/Sandbox/SandboxedEditing/Source/SandboxedEditing/Private/Features/Browser/Widgets/Shared/SSandboxDescription.h b/Engine/Plugins/Developer/Sandbox/SandboxedEditing/Source/SandboxedEditing/Private/Features/Browser/Widgets/Shared/SSandboxDescription.h
new file mode 100644
--- /dev/null
+++ b/Engine/Plugins/Developer/Sandbox/SandboxedEditing/Source/SandboxedEditing/Private/Features/Browser/Widgets/Shared/SSandboxDescription.h
@@ -0,0 +1,43 @@
+﻿// Copyright Epic Games, Inc. All Rights Reserved.
+
+#pragma once
+
+#include "Framework/SlateDelegates.h"
+#include "Widgets/DeclarativeSyntaxSupport.h"
+#include "Widgets/SCompoundWidget.h"
+
+class SExpandableArea;
+class SMultiLineEditableTextBox;
+
+namespace UE::SandboxedEditing
+{
+class FSandboxMetaDataViewModel;
+
+/** Allows editing a description. The description is displayed in an expandable area to save space. */
+class SSandboxDescription : public SCompoundWidget
+{
+public:
+	
+	SLATE_BEGIN_ARGS(SSandboxDescription) {}
+		/** Path to the sandbox for which the description is being edited. Returning unset disables the edit box. */
+		SLATE_ATTRIBUTE(TOptional<FString>, SandboxPath)
+	SLATE_END_ARGS()
+	
+	void Construct(const FArguments& InArgs, const TSharedRef<FSandboxMetaDataViewModel>& InMetaDataViewModel);
+	
+private:
+	
+	/** Used to edit the description. */
+	TSharedPtr<FSandboxMetaDataViewModel> MetaDataViewModel;
+	
+	/** Path to the sandbox for which the description is being edited. */
+	TAttribute<TOptional<FString>> SandboxPathAttr;
+	
+	/** Contains the description text */
+	TSharedPtr<SMultiLineEditableTextBox> DescriptionTextCtrl;
+	
+	FText GetDescription() const;
+	void OnDescriptionCommitted(const FText& InNewDescription, ETextCommit::Type InType) const;
+};
+}
+
diff --git a/Engine/Plugins/Developer/Sandbox/SandboxedEditing/Source/SandboxedEditing/Private/Features/Browser/Widgets/Unsandboxed/GenericBrowser/SSessionBrowserFrame.cpp b/Engine/Plugins/Developer/Sandbox/SandboxedEditing/Source/SandboxedEditing/Private/Features/Browser/Widgets/Unsandboxed/GenericBrowser/SSessionBrowserFrame.cpp
new file mode 100644
--- /dev/null
+++ b/Engine/Plugins/Developer/Sandbox/SandboxedEditing/Source/SandboxedEditing/Private/Features/Browser/Widgets/Unsandboxed/GenericBrowser/SSessionBrowserFrame.cpp
@@ -0,0 +1,107 @@
+﻿// Copyright Epic Games, Inc. All Rights Reserved.
+
+#include "SSessionBrowserFrame.h"
+
+#include "Styling/AppStyle.h"
+#include "Styling/ToolBarStyle.h"
+#include "Widgets/Layout/SBorder.h"
+#include "Widgets/SBoxPanel.h"
+#include "Widgets/Layout/SExpandableArea.h"
+#include "Widgets/Layout/SSeparator.h"
+#include "Widgets/Text/STextBlock.h"
+
+#define LOCTEXT_NAMESPACE "SSessionBrowserFrame"
+
+namespace UE::SandboxedEditing
+{
+void SSessionBrowserFrame::Construct(const FArguments& InArgs)
+{
+	ChildSlot
+	[
+		SNew(SBorder)
+		.BorderImage(FAppStyle::Get().GetBrush("ToolPanel.GroupBorder"))
+		.Padding(FMargin(1.0f, 2.0f))
+		[
+			CreateContent(InArgs)
+		]
+	];
+}
+
+TSharedRef<SVerticalBox> SSessionBrowserFrame::CreateContent(const FArguments& InArgs)
+{
+	const TSharedRef<SVerticalBox> Content = SNew(SVerticalBox);
+	
+	AppendControlWidget(InArgs, Content);
+	AppendSearchWidget(InArgs, Content);
+	
+	checkf(InArgs._SessionContent.Widget != SNullWidget::NullWidget, TEXT("This is a required argument."));
+	Content->AddSlot()
+	.FillHeight(1.0f)
+	.Padding(1.0f, 2.0f)
+	[
+		InArgs._SessionContent.Widget
+	];
+	
+	AppendDetailsWidget(InArgs, Content);
+	
+	return Content;
+}
+
+void SSessionBrowserFrame::AppendControlWidget(const FArguments& InArgs, const TSharedRef<SVerticalBox>& Content)
+{
+	const bool bHasControlWidget = InArgs._ControlContent.Widget != SNullWidget::NullWidget; 
+	if (bHasControlWidget)
+	{
+		Content->AddSlot()
+		.AutoHeight()
+		[
+			InArgs._ControlContent.Widget
+		];
+
+		Content->AddSlot()
+		.AutoHeight()
+		[
+			SNew(SSeparator)
+			.Thickness(FCoreStyle::Get().GetWidgetStyle<FToolBarStyle>("ToolBar").SeparatorThickness)
+			.SeparatorImage(&FCoreStyle::Get().GetWidgetStyle<FToolBarStyle>("ToolBar").SeparatorBrush)
+		];
+	}
+}
+
+void SSessionBrowserFrame::AppendSearchWidget(const FArguments& InArgs, const TSharedRef<SVerticalBox>& Content)
+{
+	const bool bHasSearchWidget = InArgs._SearchContent.Widget != SNullWidget::NullWidget; 
+	if (bHasSearchWidget)
+	{
+		Content->AddSlot()
+		.AutoHeight()
+		.Padding(0.f, 4.f, 0.f, 4.f)
+		[
+			InArgs._SearchContent.Widget
+		];
+	}
+}
+	
+void SSessionBrowserFrame::AppendDetailsWidget(const FArguments& InArgs, const TSharedRef<SVerticalBox>& Content)
+{
+	const bool bHasDetailsWidget = InArgs._DetailsContent.Widget != SNullWidget::NullWidget;
+	if (bHasDetailsWidget)
+	{
+		Content->AddSlot()
+	   .AutoHeight()
+	   .Padding(FMargin(2.0f, 0.0f))
+		[
+			SNew(SSeparator)
+		];
+	
+		Content->AddSlot()
+		.AutoHeight()
+		.Padding(FMargin(2.0f, 0.0f))
+		[
+			InArgs._DetailsContent.Widget
+		];
+	}
+}
+}
+
+#undef LOCTEXT_NAMESPACE
\ No newline at end of file
diff --git a/Engine/Plugins/Developer/Sandbox/SandboxedEditing/Source/SandboxedEditing/Private/Features/Browser/Widgets/Unsandboxed/GenericBrowser/SSessionBrowserFrame.h b/Engine/Plugins/Developer/Sandbox/SandboxedEditing/Source/SandboxedEditing/Private/Features/Browser/Widgets/Unsandboxed/GenericBrowser/SSessionBrowserFrame.h
new file mode 100644
--- /dev/null
+++ b/Engine/Plugins/Developer/Sandbox/SandboxedEditing/Source/SandboxedEditing/Private/Features/Browser/Widgets/Unsandboxed/GenericBrowser/SSessionBrowserFrame.h
@@ -0,0 +1,51 @@
+﻿// Copyright Epic Games, Inc. All Rights Reserved.
+
+#pragma once
+
+#include "Widgets/DeclarativeSyntaxSupport.h"
+#include "Widgets/SCompoundWidget.h"
+
+class SVerticalBox;
+
+namespace UE::SandboxedEditing
+{
+/** 
+ * Generic widget that lays out widgets related for session management.
+ * The concept of a session is kept abstract and depends on the code that uses this widget.
+ * 
+ * For now only used for Sandboxed Editing. In the future, it should be used across Multi-User, Multi-User Server, and Sandboxed editing.
+ * The purpose is to have a central place that organizes the layout of the specific widgets, and avoid code duplication.
+ * 
+ * The follow elements are laid out:
+ * - Controls: buttons to control the sessions
+ * - Search box: filtering the sessions
+ * - Tree / List view: displays the sessions
+ * - Extra info: displays more info about the session
+ */
+class SSessionBrowserFrame : public SCompoundWidget
+{
+public:
+
+	SLATE_BEGIN_ARGS(SSessionBrowserFrame){}
+		SLATE_NAMED_SLOT(FArguments, ControlContent)
+		SLATE_NAMED_SLOT(FArguments, SearchContent)
+		SLATE_NAMED_SLOT(FArguments, SessionContent)
+		SLATE_NAMED_SLOT(FArguments, DetailsContent)
+	SLATE_END_ARGS()
+
+	void Construct(const FArguments& InArgs);
+	
+private:
+	
+	/** Creates the layout */
+	TSharedRef<SVerticalBox> CreateContent(const FArguments& InArgs);
+	
+	/** Optionally adds the controls (above search), if specified by args. */
+	void AppendControlWidget(const FArguments& InArgs, const TSharedRef<SVerticalBox>& Content);
+	/** Optionally adds the search widget (above session content), if specified by args. */
+	void AppendSearchWidget(const FArguments& InArgs, const TSharedRef<SVerticalBox>& Content);
+	/** Optionally adds the controls (below session content), if specified by args. */
+	void AppendDetailsWidget(const FArguments& InArgs, const TSharedRef<SVerticalBox>& Content);
+};
+}
+
diff --git a/Engine/Plugins/Developer/Sandbox/SandboxedEditing/Source/SandboxedEditing/Private/Features/Browser/Widgets/Unsandboxed/List/Columns/LastModifiedSandboxColumn.cpp b/Engine/Plugins/Developer/Sandbox/SandboxedEditing/Source/SandboxedEditing/Private/Features/Browser/Widgets/Unsandboxed/List/Columns/LastModifiedSandboxColumn.cpp
new file mode 100644
--- /dev/null
+++ b/Engine/Plugins/Developer/Sandbox/SandboxedEditing/Source/SandboxedEditing/Private/Features/Browser/Widgets/Unsandboxed/List/Columns/LastModifiedSandboxColumn.cpp
@@ -0,0 +1,52 @@
+﻿// Copyright Epic Games, Inc. All Rights Reserved.
+
+#include "LastModifiedSandboxColumn.h"
+
+#include "Features/Browser/ViewModels/List/SandboxListItem.h"
+#include "Framework/Models/SandboxInfo.h"
+#include "SandboxedEditingStyle.h"
+#include "Widgets/Layout/SBox.h"
+#include "Widgets/Text/STextBlock.h"
+#include "Widgets/Views/SHeaderRow.h"
+
+#define LOCTEXT_NAMESPACE "FLastModifiedSandboxColumn"
+
+namespace UE::SandboxedEditing
+{
+namespace LastModifiedDetail
+{
+static FText FormatLastModified(const FSandboxInfo& InRowData)
+{
+	return FText::AsDateTime(InRowData.LastModified);
+}
+}
+void FLastModifiedSandboxColumn::PopulateSearchTerms(const TSharedRef<FSandboxListItem>& InRowData, TArray<FString>& OutSearchTerms) const
+{
+	OutSearchTerms.Add(
+		LastModifiedDetail::FormatLastModified(InRowData->SandboxInfo).ToString()
+		);
+}
+
+SHeaderRow::FColumn::FArguments FLastModifiedSandboxColumn::MakeColumnArguments()
+{
+	return SHeaderRow::FColumn::FArguments()
+		.ColumnId(LastModifiedSandboxColumn)
+		.FillWidth(FSandboxedEditingStyle::Get().GetFloat("SandboxedEditing.Browser.LastModifiedColumn.FillWidth"))
+		.DefaultLabel(LOCTEXT("LastModified.Label", "Last Modified"))
+		.ToolTipText(LOCTEXT("LastModified.Description", "When this sandbox was last modified"));
+}
+
+TSharedRef<SWidget> FLastModifiedSandboxColumn::MakeColumnWidget(const FMakeColumnWidgetArgs& InArgs)
+{
+	return SNew(SBox)
+		.VAlign(VAlign_Center)
+		.Padding(FSandboxedEditingStyle::Get().GetMargin("SandboxedEditing.Browser.LastModifiedColumn.Padding"))
+		[
+			SNew(STextBlock)
+			.Text_Lambda([RowData = InArgs.RowData](){ return LastModifiedDetail::FormatLastModified(RowData->SandboxInfo); })
+			.HighlightText(InArgs.HighlightText)
+		];
+}
+}
+
+#undef LOCTEXT_NAMESPACE
\ No newline at end of file
diff --git a/Engine/Plugins/Developer/Sandbox/SandboxedEditing/Source/SandboxedEditing/Private/Features/Browser/Widgets/Unsandboxed/List/Columns/LastModifiedSandboxColumn.h b/Engine/Plugins/Developer/Sandbox/SandboxedEditing/Source/SandboxedEditing/Private/Features/Browser/Widgets/Unsandboxed/List/Columns/LastModifiedSandboxColumn.h
new file mode 100644
--- /dev/null
+++ b/Engine/Plugins/Developer/Sandbox/SandboxedEditing/Source/SandboxedEditing/Private/Features/Browser/Widgets/Unsandboxed/List/Columns/LastModifiedSandboxColumn.h
@@ -0,0 +1,24 @@
+﻿// Copyright Epic Games, Inc. All Rights Reserved.
+
+#pragma once
+
+#include "SandboxColumns.h"
+
+namespace UE::SandboxedEditing
+{
+/** Displays when the sandbox was last modified. */
+class FLastModifiedSandboxColumn : public ISandboxColumnBehavior, public ISandboxWidgetFactory
+{
+public:
+	
+	//~ Begin ISandboxColumnBehavior Interface
+	virtual void PopulateSearchTerms(const TSharedRef<FSandboxListItem>& InRowData, TArray<FString>& OutSearchTerms) const override;
+	//~ End ISandboxColumnBehavior Interface
+	
+	//~ Begin ISandboxWidgetFactory Interface
+	virtual SHeaderRow::FColumn::FArguments MakeColumnArguments() override;
+	virtual TSharedRef<SWidget> MakeColumnWidget(const FMakeColumnWidgetArgs& InArgs) override;
+	//~ End ISandboxWidgetFactory Interface
+};
+}
+
diff --git a/Engine/Plugins/Developer/Sandbox/SandboxedEditing/Source/SandboxedEditing/Private/Features/Browser/Widgets/Unsandboxed/List/Columns/NameSandboxColumn.cpp b/Engine/Plugins/Developer/Sandbox/SandboxedEditing/Source/SandboxedEditing/Private/Features/Browser/Widgets/Unsandboxed/List/Columns/NameSandboxColumn.cpp
new file mode 100644
--- /dev/null
+++ b/Engine/Plugins/Developer/Sandbox/SandboxedEditing/Source/SandboxedEditing/Private/Features/Browser/Widgets/Unsandboxed/List/Columns/NameSandboxColumn.cpp
@@ -0,0 +1,56 @@
+﻿// Copyright Epic Games, Inc. All Rights Reserved.
+
+#include "NameSandboxColumn.h"
+
+#include "Features/Browser/ViewModels/List/SandboxListItem.h"
+#include "Framework/Models/SandboxInfo.h"
+#include "SandboxedEditingStyle.h"
+#include "Widgets/Layout/SBox.h"
+#include "Widgets/Text/SInlineEditableTextBlock.h"
+#include "Widgets/Views/SHeaderRow.h"
+
+#define LOCTEXT_NAMESPACE "FNameSandboxColumn"
+
+namespace UE::SandboxedEditing
+{
+void FNameSandboxColumn::PopulateSearchTerms(const TSharedRef<FSandboxListItem>& InRowData, TArray<FString>& OutSearchTerms) const
+{
+	OutSearchTerms.Add(InRowData->SandboxInfo.Name);
+}
+
+SHeaderRow::FColumn::FArguments FNameSandboxColumn::MakeColumnArguments()
+{
+	return SHeaderRow::FColumn::FArguments()
+		.ColumnId(NameSandboxColumn)
+		.FillWidth(FSandboxedEditingStyle::Get().GetFloat("SandboxedEditing.Browser.NameColumn.FillWidth"))
+		.DefaultLabel(LOCTEXT("Name.Label", "Sandbox"))
+		.ToolTipText(LOCTEXT("Name.Description", "The name of this sandbox"))
+		.ShouldGenerateWidget(true);
+}
+
+TSharedRef<SWidget> FNameSandboxColumn::MakeColumnWidget(const FMakeColumnWidgetArgs& InArgs)
+{
+	return SNew(SBox)
+		.VAlign(VAlign_Center)
+		.Padding(FSandboxedEditingStyle::Get().GetMargin("SandboxedEditing.Browser.NameColumn.Padding"))
+		[
+			SNew(SInlineEditableTextBlock)
+			.Text_Lambda([Name = InArgs.RowData->SandboxInfo.Name] { return FText::AsCultureInvariant(Name); })
+			.ToolTipText_Lambda([Description = InArgs.RowData->SandboxInfo.Description]
+			{
+				return !Description.IsEmpty()
+					? FText::AsCultureInvariant(Description) 
+					: LOCTEXT("NoDescription", "This sandbox has no user-provided description");
+			})
+			.HighlightText(InArgs.HighlightText)
+			//.OnTextCommitted(this, &SSessionRow::OnSessionNameCommitted) TODO DP: Implement renaming
+			.IsReadOnly(false)
+			.IsSelected(InArgs.IsSelectedDelegate)
+			//.OnVerifyTextChanged(this, &SSessionRow::OnValidatingSessionName) TODO DP: Implement renaming
+			//.Font(FontInfo) TODO DP: Change font based on whether active
+			//.ColorAndOpacity(FontColor) TODO DP: Change font based on whether active
+		];
+}
+}
+
+#undef LOCTEXT_NAMESPACE
\ No newline at end of file
diff --git a/Engine/Plugins/Developer/Sandbox/SandboxedEditing/Source/SandboxedEditing/Private/Features/Browser/Widgets/Unsandboxed/List/Columns/NameSandboxColumn.h b/Engine/Plugins/Developer/Sandbox/SandboxedEditing/Source/SandboxedEditing/Private/Features/Browser/Widgets/Unsandboxed/List/Columns/NameSandboxColumn.h
new file mode 100644
--- /dev/null
+++ b/Engine/Plugins/Developer/Sandbox/SandboxedEditing/Source/SandboxedEditing/Private/Features/Browser/Widgets/Unsandboxed/List/Columns/NameSandboxColumn.h
@@ -0,0 +1,24 @@
+﻿// Copyright Epic Games, Inc. All Rights Reserved.
+
+#pragma once
+
+#include "SandboxColumns.h"
+
+namespace UE::SandboxedEditing
+{
+/** Displays the sandbox name. */
+class FNameSandboxColumn : public ISandboxColumnBehavior, public ISandboxWidgetFactory
+{
+public:
+	
+	//~ Begin ISandboxColumnBehavior Interface
+	virtual void PopulateSearchTerms(const TSharedRef<FSandboxListItem>& InRowData, TArray<FString>& OutSearchTerms) const override;
+	//~ End ISandboxColumnBehavior Interface
+	
+	//~ Begin ISandboxWidgetFactory Interface
+	virtual SHeaderRow::FColumn::FArguments MakeColumnArguments() override;
+	virtual TSharedRef<SWidget> MakeColumnWidget(const FMakeColumnWidgetArgs& InArgs) override;
+	//~ End ISandboxWidgetFactory Interface
+};
+}
+
diff --git a/Engine/Plugins/Developer/Sandbox/SandboxedEditing/Source/SandboxedEditing/Private/Features/Browser/Widgets/Unsandboxed/List/Columns/SandboxColumns.cpp b/Engine/Plugins/Developer/Sandbox/SandboxedEditing/Source/SandboxedEditing/Private/Features/Browser/Widgets/Unsandboxed/List/Columns/SandboxColumns.cpp
new file mode 100644
--- /dev/null
+++ b/Engine/Plugins/Developer/Sandbox/SandboxedEditing/Source/SandboxedEditing/Private/Features/Browser/Widgets/Unsandboxed/List/Columns/SandboxColumns.cpp
@@ -0,0 +1,24 @@
+﻿// Copyright Epic Games, Inc. All Rights Reserved.
+
+#include "SandboxColumns.h"
+
+#include "LastModifiedSandboxColumn.h"
+#include "NameSandboxColumn.h"
+#include "VersionSandboxColumn.h"
+
+namespace UE::SandboxedEditing
+{
+TMap<FName, FSandboxColumnInfo> MakeBuiltInColumns()
+{
+	const TSharedRef<FNameSandboxColumn> Name = MakeShared<FNameSandboxColumn>();
+	const TSharedRef<FVersionSandboxColumn> Version = MakeShared<FVersionSandboxColumn>();
+	const TSharedRef<FLastModifiedSandboxColumn> LastModified = MakeShared<FLastModifiedSandboxColumn>();
+	const TMap<FName, FSandboxColumnInfo> Result
+	{
+		{ NameSandboxColumn, FSandboxColumnInfo(Name, Name) },
+		{ VersionSandboxNameColumn, FSandboxColumnInfo(Version, Version) },
+		{ LastModifiedSandboxColumn, FSandboxColumnInfo(LastModified, LastModified) }
+	};
+	return Result;
+}
+}
diff --git a/Engine/Plugins/Developer/Sandbox/SandboxedEditing/Source/SandboxedEditing/Private/Features/Browser/Widgets/Unsandboxed/List/Columns/SandboxColumns.h b/Engine/Plugins/Developer/Sandbox/SandboxedEditing/Source/SandboxedEditing/Private/Features/Browser/Widgets/Unsandboxed/List/Columns/SandboxColumns.h
new file mode 100644
--- /dev/null
+++ b/Engine/Plugins/Developer/Sandbox/SandboxedEditing/Source/SandboxedEditing/Private/Features/Browser/Widgets/Unsandboxed/List/Columns/SandboxColumns.h
@@ -0,0 +1,26 @@
+﻿// Copyright Epic Games, Inc. All Rights Reserved.
+
+#pragma once
+
+#include "Features/Browser/ViewModels/Columns/ColumnInfo.h"
+#include "Templates/SharedPointer.h"
+#include "UObject/NameTypes.h"
+
+namespace UE::SandboxedEditing
+{
+class FSandboxListItem;
+
+using ISandboxColumnBehavior = IColumnBehavior<TSharedRef<FSandboxListItem>>;
+using ISandboxWidgetFactory = IColumnWidgetFactory<TSharedRef<FSandboxListItem>>;
+using FSandboxColumnInfo = TColumnInfo<TSharedRef<FSandboxListItem>>;
+
+/** Needed for SLATE_ARGUMENT (the macro is written in such a way that TMap is not parsed correctly). */
+using FSandboxColumnFactoryMap = TMap<FName, TSharedRef<ISandboxWidgetFactory>>;
+
+constexpr FLazyName NameSandboxColumn(TEXT("NameSandboxColumn"));
+constexpr FLazyName VersionSandboxNameColumn(TEXT("VersionSandboxNameColumn"));
+constexpr FLazyName LastModifiedSandboxColumn(TEXT("LastModifiedSandboxColumn"));
+
+/** @return The array of built-in columns. */
+TMap<FName, FSandboxColumnInfo> MakeBuiltInColumns();
+}
diff --git a/Engine/Plugins/Developer/Sandbox/SandboxedEditing/Source/SandboxedEditing/Private/Features/Browser/Widgets/Unsandboxed/List/Columns/VersionSandboxColumn.cpp b/Engine/Plugins/Developer/Sandbox/SandboxedEditing/Source/SandboxedEditing/Private/Features/Browser/Widgets/Unsandboxed/List/Columns/VersionSandboxColumn.cpp
new file mode 100644
--- /dev/null
+++ b/Engine/Plugins/Developer/Sandbox/SandboxedEditing/Source/SandboxedEditing/Private/Features/Browser/Widgets/Unsandboxed/List/Columns/VersionSandboxColumn.cpp
@@ -0,0 +1,52 @@
+﻿// Copyright Epic Games, Inc. All Rights Reserved.
+
+#include "VersionSandboxColumn.h"
+
+#include "Features/Browser/ViewModels/List/SandboxListItem.h"
+#include "Framework/Models/SandboxInfo.h"
+#include "SandboxedEditingStyle.h"
+#include "Widgets/Layout/SBox.h"
+#include "Widgets/Text/STextBlock.h"
+#include "Widgets/Views/SHeaderRow.h"
+
+#define LOCTEXT_NAMESPACE "FVersionSandboxColumn"
+
+namespace UE::SandboxedEditing
+{
+namespace VersionDetail
+{
+static FText FormatVersion(const FSandboxInfo& InRowData)
+{
+	return FText::AsCultureInvariant(InRowData.EngineVersion.ToString());
+}
+}
+void FVersionSandboxColumn::PopulateSearchTerms(const TSharedRef<FSandboxListItem>& InRowData, TArray<FString>& OutSearchTerms) const
+{
+	OutSearchTerms.Add(
+		VersionDetail::FormatVersion(InRowData->SandboxInfo).ToString()
+		);
+}
+
+SHeaderRow::FColumn::FArguments FVersionSandboxColumn::MakeColumnArguments()
+{
+	return SHeaderRow::FColumn::FArguments()
+		.ColumnId(VersionSandboxNameColumn)
+		.FillWidth(FSandboxedEditingStyle::Get().GetFloat("SandboxedEditing.Browser.VersionColumn.FillWidth"))
+		.DefaultLabel(LOCTEXT("Version.Label", "Engine Version"))
+		.ToolTipText(LOCTEXT("Version.Description", "The engine version this sandbox had when it was created."));
+}
+
+TSharedRef<SWidget> FVersionSandboxColumn::MakeColumnWidget(const FMakeColumnWidgetArgs& InArgs)
+{
+	return SNew(SBox)
+		.VAlign(VAlign_Center)
+		.Padding(FSandboxedEditingStyle::Get().GetMargin("SandboxedEditing.Browser.VersionColumn.Padding"))
+		[
+			SNew(STextBlock)
+			.Text_Lambda([RowData = InArgs.RowData](){ return VersionDetail::FormatVersion(RowData->SandboxInfo); })
+			.HighlightText(InArgs.HighlightText)
+		];
+}
+}
+
+#undef LOCTEXT_NAMESPACE
\ No newline at end of file
diff --git a/Engine/Plugins/Developer/Sandbox/SandboxedEditing/Source/SandboxedEditing/Private/Features/Browser/Widgets/Unsandboxed/List/Columns/VersionSandboxColumn.h b/Engine/Plugins/Developer/Sandbox/SandboxedEditing/Source/SandboxedEditing/Private/Features/Browser/Widgets/Unsandboxed/List/Columns/VersionSandboxColumn.h
new file mode 100644
--- /dev/null
+++ b/Engine/Plugins/Developer/Sandbox/SandboxedEditing/Source/SandboxedEditing/Private/Features/Browser/Widgets/Unsandboxed/List/Columns/VersionSandboxColumn.h
@@ -0,0 +1,24 @@
+﻿// Copyright Epic Games, Inc. All Rights Reserved.
+
+#pragma once
+
+#include "SandboxColumns.h"
+
+namespace UE::SandboxedEditing
+{
+/** Displays the engine version this sandbox was last used in. */
+class FVersionSandboxColumn : public ISandboxColumnBehavior, public ISandboxWidgetFactory
+{
+public:
+	
+	//~ Begin ISandboxColumnBehavior Interface
+	virtual void PopulateSearchTerms(const TSharedRef<FSandboxListItem>& InRowData, TArray<FString>& OutSearchTerms) const override;
+	//~ End ISandboxColumnBehavior Interface
+	
+	//~ Begin ISandboxWidgetFactory Interface
+	virtual SHeaderRow::FColumn::FArguments MakeColumnArguments() override;
+	virtual TSharedRef<SWidget> MakeColumnWidget(const FMakeColumnWidgetArgs& InArgs) override;
+	//~ End ISandboxWidgetFactory Interface
+};
+}
+
diff --git a/Engine/Plugins/Developer/Sandbox/SandboxedEditing/Source/SandboxedEditing/Private/Features/Browser/Widgets/Unsandboxed/List/DummySandboxItem.h b/Engine/Plugins/Developer/Sandbox/SandboxedEditing/Source/SandboxedEditing/Private/Features/Browser/Widgets/Unsandboxed/List/DummySandboxItem.h
new file mode 100644
--- /dev/null
+++ b/Engine/Plugins/Developer/Sandbox/SandboxedEditing/Source/SandboxedEditing/Private/Features/Browser/Widgets/Unsandboxed/List/DummySandboxItem.h
@@ -0,0 +1,32 @@
+﻿// Copyright Epic Games, Inc. All Rights Reserved.
+
+#pragma once
+
+#include "Features/Browser/ViewModels/SandboxControlsViewModel.h"
+#include "HAL/Platform.h"
+
+namespace UE::SandboxedEditing
+{
+class FSandboxListItem;
+
+/** The intent of the dummy sandbox item. */
+enum class EDummySandboxItemType : uint8
+{
+	CreateSandbox
+};
+
+/** 
+ * Holds info about a FSandboxListItem added to the SListView but does not really represent a real sandbox. 
+ * It's a way to get the SListView to display row widget.
+ */
+struct FDummySandboxItemInfo
+{
+	const TSharedRef<FSandboxListItem> Item;
+	const EDummySandboxItemType Type;
+
+	explicit FDummySandboxItemInfo(const TSharedRef<FSandboxListItem>& InItem, EDummySandboxItemType InType)
+		: Item(InItem)
+		, Type(InType)
+	{}
+};
+}
\ No newline at end of file
diff --git a/Engine/Plugins/Developer/Sandbox/SandboxedEditing/Source/SandboxedEditing/Private/Features/Browser/Widgets/Unsandboxed/List/SCreateSandboxListRow.cpp b/Engine/Plugins/Developer/Sandbox/SandboxedEditing/Source/SandboxedEditing/Private/Features/Browser/Widgets/Unsandboxed/List/SCreateSandboxListRow.cpp
new file mode 100644
--- /dev/null
+++ b/Engine/Plugins/Developer/Sandbox/SandboxedEditing/Source/SandboxedEditing/Private/Features/Browser/Widgets/Unsandboxed/List/SCreateSandboxListRow.cpp
@@ -0,0 +1,163 @@
+﻿// Copyright Epic Games, Inc. All Rights Reserved.
+
+#include "SCreateSandboxListRow.h"
+
+#include "Columns/SandboxColumns.h"
+#include "Features/Browser/ViewModels/SandboxCreationWorkflow.h"
+#include "Framework/Application/SlateApplication.h"
+#include "Utils/WidgetUtils.h"
+#include "Widgets/Images/SImage.h"
+#include "Widgets/Input/SEditableTextBox.h"
+#include "Widgets/Layout/SUniformGridPanel.h"
+
+#define LOCTEXT_NAMESPACE "SCreateSandboxListRow"
+
+namespace UE::SandboxedEditing
+{
+void SCreateSandboxListRow::Construct(const FArguments& InArgs, const TSharedRef<STableViewBase>& InOwnerTable)
+{
+	WorkflowAttr = InArgs._Workflow;
+	
+	SMultiColumnTableRow::Construct(
+		SMultiColumnTableRow::FArguments().Padding(InArgs._Padding),
+		InOwnerTable
+		);
+}
+
+void SCreateSandboxListRow::Tick(const FGeometry& AllottedGeometry, const double InCurrentTime, const float InDeltaTime)
+{
+	SMultiColumnTableRow<TSharedPtr<FSandboxListItem>>::Tick(AllottedGeometry, InCurrentTime, InDeltaTime);
+	
+	if (NameWidget && !bHasTakenInitialFocus)
+	{
+		bHasTakenInitialFocus = FSlateApplication::Get().SetKeyboardFocus(NameWidget.ToSharedRef());
+	}
+}
+
+TSharedRef<SWidget> SCreateSandboxListRow::GenerateWidgetForColumn(const FName& InColumnName)
+{
+	if (InColumnName == NameSandboxColumn)
+	{
+		return SNew(SHorizontalBox)
+
+			+SHorizontalBox::Slot()
+			.AutoWidth()
+			.HAlign(HAlign_Center)
+			.VAlign(VAlign_Center)
+			.Padding(1.f, 1.f, 4.f, 1.f)
+			[
+				SNew(SImage)
+				.Image(FAppStyle::GetBrush("Icons.PlusCircle"))
+			]
+			
+			+SHorizontalBox::Slot()
+			[
+				CreateNameWidget()
+			]
+
+			+SHorizontalBox::Slot()
+			.AutoWidth()
+			.Padding(2.0f)
+			.HAlign(HAlign_Left)
+			[
+				CreateButtons()
+			];
+	}
+	
+	return SNullWidget::NullWidget;
+}
+
+TSharedRef<SWidget> SCreateSandboxListRow::CreateNameWidget()
+{
+	return SAssignNew(NameWidget, SEditableTextBox)
+		.HintText(LOCTEXT("EnterSessionNameHint", "Enter sandbox name"))
+		.OnTextCommitted(this, &SCreateSandboxListRow::OnSandboxNameCommitted)
+		.OnKeyDownHandler(this, &SCreateSandboxListRow::OnKeyDownHandler)
+		.OnTextChanged(this, &SCreateSandboxListRow::OnSandboxNameChanged);
+}
+
+TSharedRef<SWidget> SCreateSandboxListRow::CreateButtons()
+{
+	return SNew(SUniformGridPanel)
+		.SlotPadding(FMargin(1.0f, 0.0f))
+
+		// 'Accept' button
+		+SUniformGridPanel::Slot(0, 0)
+		[
+			MakePositiveActionButton(
+				FAppStyle::GetBrush("Icons.Check"),
+				LOCTEXT("CreateCheckIconTooltip", "Create the session"),
+				TAttribute<bool>::Create([this]
+				{
+					const FSandboxCreationWorkflow* Workflow = GetWorkflow();
+					return Workflow && Workflow->CanConfirm();
+				}),
+				FOnClicked::CreateLambda([this]
+				{
+					if (FSandboxCreationWorkflow* Workflow = GetWorkflow())
+					{
+						Workflow->Confirm();
+					}
+					return FReply::Handled();
+				}))
+		]
+
+		// 'Decline' button
+		+SUniformGridPanel::Slot(1, 0)
+		[
+			MakeNegativeActionButton(
+				FAppStyle::GetBrush("Icons.X"),
+				LOCTEXT("CancelIconTooltip", "Cancel"),
+				true, // Always enabled.
+				FOnClicked::CreateLambda([this]
+				{
+					if (FSandboxCreationWorkflow* Workflow = GetWorkflow())
+					{
+						Workflow->Cancel();
+					}
+					return FReply::Handled();
+				}))
+		];
+}
+
+void SCreateSandboxListRow::OnSandboxNameCommitted(const FText& InNewText, ETextCommit::Type InCommitType)
+{
+	FSandboxCreationWorkflow* Workflow = GetWorkflow();
+	if (!Workflow)
+	{
+		return;
+	}
+	
+	Workflow->SetName(InNewText.ToString());
+	if (Workflow->CanConfirm() && InCommitType == ETextCommit::OnEnter)
+	{
+		Workflow->Confirm();
+	}
+}
+
+FReply SCreateSandboxListRow::OnKeyDownHandler(const FGeometry&, const FKeyEvent& InKeyEvent)
+{
+	FSandboxCreationWorkflow* Workflow = GetWorkflow();
+	if (InKeyEvent.GetKey() == EKeys::Escape && Workflow)
+	{
+		Workflow->Cancel();
+		return FReply::Handled();
+	}
+	return FReply::Unhandled();
+}
+
+void SCreateSandboxListRow::OnSandboxNameChanged(const FText& InNewText)
+{
+	if (FSandboxCreationWorkflow* Workflow = GetWorkflow())
+	{
+		const FString Name = InNewText.ToString();
+		Workflow->SetName(Name);
+		
+		FText Error = FText::GetEmpty();
+		Workflow->IsValidName(Name, &Error);
+		NameWidget->SetError(Error);
+	}
+}
+}
+
+#undef LOCTEXT_NAMESPACE
\ No newline at end of file
diff --git a/Engine/Plugins/Developer/Sandbox/SandboxedEditing/Source/SandboxedEditing/Private/Features/Browser/Widgets/Unsandboxed/List/SCreateSandboxListRow.h b/Engine/Plugins/Developer/Sandbox/SandboxedEditing/Source/SandboxedEditing/Private/Features/Browser/Widgets/Unsandboxed/List/SCreateSandboxListRow.h
new file mode 100644
--- /dev/null
+++ b/Engine/Plugins/Developer/Sandbox/SandboxedEditing/Source/SandboxedEditing/Private/Features/Browser/Widgets/Unsandboxed/List/SCreateSandboxListRow.h
@@ -0,0 +1,64 @@
+﻿// Copyright Epic Games, Inc. All Rights Reserved.
+
+#pragma once
+
+#include "Widgets/DeclarativeSyntaxSupport.h"
+#include "Widgets/Views/SListView.h"
+#include "Widgets/Views/STableRow.h"
+
+class SEditableTextBox;
+
+namespace UE::SandboxedEditing
+{
+class FSandboxCreationWorkflow;
+class FSandboxListItem;
+
+/** A row in the SSandboxListView for creating a new sandbox. */
+class SCreateSandboxListRow : public SMultiColumnTableRow<TSharedPtr<FSandboxListItem>>
+{
+public:
+	
+	SLATE_BEGIN_ARGS(SCreateSandboxListRow) {}
+		/** Gets the creation workflow this widget displays. */
+		SLATE_ATTRIBUTE(FSandboxCreationWorkflow*, Workflow)
+		
+		/** Padding for the row */
+		SLATE_ARGUMENT(FMargin, Padding)
+	SLATE_END_ARGS()
+
+	void Construct(const FArguments& InArgs, const TSharedRef<STableViewBase>& InOwnerTable);
+	
+	virtual void Tick(const FGeometry& AllottedGeometry, const double InCurrentTime, const float InDeltaTime) override;
+
+	//~ Begin  SMultiColumnTableRow Interface
+	virtual TSharedRef<SWidget> GenerateWidgetForColumn(const FName& InColumnName) override;
+	//~ End  SMultiColumnTableRow Interface
+	
+private:
+	
+	/** Gets the creation workflow this widget displays. */
+	TAttribute<FSandboxCreationWorkflow*> WorkflowAttr;
+	
+	/** Name widget */
+	TSharedPtr<SEditableTextBox> NameWidget;
+	
+	/** Whether keyboard focus has been set to the editable text already. */
+	bool bHasTakenInitialFocus = false;
+	
+	/** @return The widget for editing the name. */
+	TSharedRef<SWidget> CreateNameWidget();
+	/** @return Confirm and cancel buttons. */
+	TSharedRef<SWidget> CreateButtons();
+	
+	/** Creates the sandbox if enter is pressed. */
+	void OnSandboxNameCommitted(const FText& InNewText, ETextCommit::Type InCommitType);
+	/** Handles cancelling when esc is pressed. */
+	FReply OnKeyDownHandler(const FGeometry& Geometry, const FKeyEvent& InKeyEvent);
+	/** Handles displaying an error. */
+	void OnSandboxNameChanged(const FText& InNewText);
+	
+	/** @return Gets the workflow being displayed. */
+	FSandboxCreationWorkflow* GetWorkflow() const { return WorkflowAttr.Get(); }
+};
+}
+
diff --git a/Engine/Plugins/Developer/Sandbox/SandboxedEditing/Source/SandboxedEditing/Private/Features/Browser/Widgets/Unsandboxed/List/SSandboxListRow.cpp b/Engine/Plugins/Developer/Sandbox/SandboxedEditing/Source/SandboxedEditing/Private/Features/Browser/Widgets/Unsandboxed/List/SSandboxListRow.cpp
new file mode 100644
--- /dev/null
+++ b/Engine/Plugins/Developer/Sandbox/SandboxedEditing/Source/SandboxedEditing/Private/Features/Browser/Widgets/Unsandboxed/List/SSandboxListRow.cpp
@@ -0,0 +1,38 @@
+﻿// Copyright Epic Games, Inc. All Rights Reserved.
+
+#include "SSandboxListRow.h"
+
+#include "Features/Browser/ViewModels/List/SandboxListItem.h"
+
+namespace UE::SandboxedEditing
+{
+void SSandboxListRow::Construct(const FArguments& InArgs, const TSharedRef<FSandboxListItem>& InItem, const TSharedRef<STableViewBase>& InOwnerTable)
+{
+	Item = InItem;
+	ColumnFactoryMap = InArgs._ColumnFactories;
+	
+	HighlightTextAttr = InArgs._HighlightText;
+	IsSelectedAttr = InArgs._IsSelected;
+	
+	SMultiColumnTableRow::Construct(
+		SMultiColumnTableRow::FArguments().Padding(InArgs._Padding),
+		InOwnerTable
+		);
+}
+
+TSharedRef<SWidget> SSandboxListRow::GenerateWidgetForColumn(const FName& InColumnName)
+{
+	if (TSharedRef<ISandboxWidgetFactory>* Factory = ColumnFactoryMap.Find(InColumnName))
+	{
+		return Factory->Get().MakeColumnWidget(
+			ISandboxWidgetFactory::FMakeColumnWidgetArgs(
+				Item.ToSharedRef(), HighlightTextAttr,
+				FIsSelected::CreateLambda([this]{ return IsSelectedAttr.Get(); })
+			)
+		);
+	}
+	
+	checkNoEntry();
+	return SNullWidget::NullWidget;
+}
+}
diff --git a/Engine/Plugins/Developer/Sandbox/SandboxedEditing/Source/SandboxedEditing/Private/Features/Browser/Widgets/Unsandboxed/List/SSandboxListRow.h b/Engine/Plugins/Developer/Sandbox/SandboxedEditing/Source/SandboxedEditing/Private/Features/Browser/Widgets/Unsandboxed/List/SSandboxListRow.h
new file mode 100644
--- /dev/null
+++ b/Engine/Plugins/Developer/Sandbox/SandboxedEditing/Source/SandboxedEditing/Private/Features/Browser/Widgets/Unsandboxed/List/SSandboxListRow.h
@@ -0,0 +1,52 @@
+﻿// Copyright Epic Games, Inc. All Rights Reserved.
+
+#pragma once
+
+#include "Columns/SandboxColumns.h"
+#include "Widgets/DeclarativeSyntaxSupport.h"
+#include "Widgets/Views/SListView.h"
+#include "Widgets/Views/STableRow.h"
+
+namespace UE::SandboxedEditing
+{
+class FSandboxListItem;
+
+/** A row in the SSandboxListView displaying a single sandbox. */
+class SSandboxListRow : public SMultiColumnTableRow<TSharedPtr<FSandboxListItem>>
+{
+public:
+	
+	SLATE_BEGIN_ARGS(SSandboxListRow) {}
+		/** The columns we know about */
+		SLATE_ARGUMENT(FSandboxColumnFactoryMap, ColumnFactories)
+		
+		/** Which text should be highlighted (text search). */
+		SLATE_ATTRIBUTE(FText, HighlightText)
+		/** Whether the item is selected. */
+		SLATE_ATTRIBUTE(bool, IsSelected)
+
+		/** Padding for the row */
+		SLATE_ARGUMENT(FMargin, Padding)
+	SLATE_END_ARGS()
+
+	void Construct(const FArguments& InArgs, const TSharedRef<FSandboxListItem>& InItem, const TSharedRef<STableViewBase>& InOwnerTable);
+
+	//~ Begin  SMultiColumnTableRow Interface
+	virtual TSharedRef<SWidget> GenerateWidgetForColumn(const FName& InColumnName) override;
+	//~ End  SMultiColumnTableRow Interface
+	
+private:
+	
+	/** Used to build the columns. */
+	FSandboxColumnFactoryMap ColumnFactoryMap;
+	
+	/** The displayed item. */
+	TSharedPtr<FSandboxListItem> Item;
+	
+	/** Which text should be highlighted (text search). */
+	TAttribute<FText> HighlightTextAttr;
+	/** Whether the item is selected. */
+	TAttribute<bool> IsSelectedAttr;
+};
+}
+
diff --git a/Engine/Plugins/Developer/Sandbox/SandboxedEditing/Source/SandboxedEditing/Private/Features/Browser/Widgets/Unsandboxed/List/SSandboxListView.cpp b/Engine/Plugins/Developer/Sandbox/SandboxedEditing/Source/SandboxedEditing/Private/Features/Browser/Widgets/Unsandboxed/List/SSandboxListView.cpp
new file mode 100644
--- /dev/null
+++ b/Engine/Plugins/Developer/Sandbox/SandboxedEditing/Source/SandboxedEditing/Private/Features/Browser/Widgets/Unsandboxed/List/SSandboxListView.cpp
@@ -0,0 +1,159 @@
+﻿// Copyright Epic Games, Inc. All Rights Reserved.
+
+#include "SSandboxListView.h"
+
+#include "Features/Browser/ViewModels/List/SandboxListItem.h"
+#include "Features/Browser/ViewModels/List/SandboxListViewModel.h"
+#include "Features/Browser/ViewModels/SandboxControlsViewModel.h"
+#include "Features/Browser/Widgets/Unsandboxed/List/SCreateSandboxListRow.h"
+#include "SSandboxListRow.h"
+#include "SandboxedEditingStyle.h"
+#include "Utils/BreakBehavior.h"
+#include "Widgets/Text/STextBlock.h"
+#include "Widgets/Views/SListView.h"
+
+#define LOCTEXT_NAMESPACE "SSandboxListView"
+
+namespace UE::SandboxedEditing
+{
+void SSandboxListView::Construct(const FArguments& InArgs, const FViewModels& InViewModels)
+{
+	ListViewModel = InViewModels.ListViewModel;
+	ControlsViewModel = InViewModels.ControlsViewModel;
+	HighlightTextAttr = InArgs._HighlightText;
+	ColumnFactoryMap = InArgs._ColumnFactories;
+	
+	ChildSlot
+	[
+		SAssignNew(ListView, SListView<TSharedPtr<FSandboxListItem>>)
+		.ListItemsSource(&Items)
+		.OnGenerateRow(this, &SSandboxListView::OnGenerateRow)
+		.OnMouseButtonDoubleClick(this, &SSandboxListView::OnItemDoubleClicked)
+		.HeaderRow(MakeHeaderRow())
+	];
+	
+	ControlsViewModel->OnCreationWorkflowStarted().AddSP(this, &SSandboxListView::OnCreationWorkflowStarted);
+	ControlsViewModel->OnCreationWorkflowEnded().AddSP(this, &SSandboxListView::OnCreationWorkflowEnded);
+	if (const FSandboxCreationWorkflow* CurrentWorkflow = ControlsViewModel->GetCurrentCreationWorkflow())
+	{
+		OnCreationWorkflowStarted(*CurrentWorkflow);
+	}
+	
+	ListViewModel->OnItemsChanged().AddSP(this, &SSandboxListView::OnItemsChanged);
+	
+	// Finally, refresh the UI to display all the current items.
+	OnItemsChanged();
+}
+
+void SSandboxListView::ForEachSelectedItem(TFunctionRef<EBreakBehavior(const FString& InRootPath)> InCallback) const
+{
+	// We should report no selection while we're in some kind of operation / workflow.
+	if (DummyItem)
+	{
+		return;
+	}
+	
+	for (const TSharedPtr<FSandboxListItem>& Item : ListView->GetSelectedItems())
+	{
+		if (InCallback(Item->SandboxInfo.SandboxRoot) == EBreakBehavior::Break)
+		{
+			break;
+		}
+	}
+}
+
+TOptional<FString> SSandboxListView::GetSelectedItem() const
+{
+	return NumSelected() == 1 ? ListView->GetSelectedItems()[0]->SandboxInfo.SandboxRoot : TOptional<FString>{};
+}
+
+int32 SSandboxListView::NumSelected() const
+{
+	return DummyItem ? 0 : ListView->GetNumItemsSelected();
+}
+
+TSharedRef<SHeaderRow> SSandboxListView::MakeHeaderRow()
+{
+	const TSharedRef<SHeaderRow> HeaderRow = SNew(SHeaderRow);
+		
+	for (FName ColumnId : ListViewModel->GetDisplayedColumns())
+	{
+		if (const TSharedRef<ISandboxWidgetFactory>* Factory = ColumnFactoryMap.Find(ColumnId))
+		{
+			HeaderRow->AddColumn(Factory->Get().MakeColumnArguments());
+		}
+	}
+	
+	return HeaderRow;
+}
+
+TSharedRef<ITableRow> SSandboxListView::OnGenerateRow(TSharedPtr<FSandboxListItem> InItem, const TSharedRef<STableViewBase>& InOwnerTable)
+{
+	check(InItem);
+	const FMargin Padding = FSandboxedEditingStyle::Get().GetMargin("SandboxedEditing.Browser.RowPadding");
+	
+	const bool bIsDummyItem = DummyItem && InItem == DummyItem->Item; 
+	if (bIsDummyItem)
+	{
+		switch (DummyItem->Type)
+		{
+		case EDummySandboxItemType::CreateSandbox: return SNew(SCreateSandboxListRow, InOwnerTable)
+			.Padding(Padding)
+			.Workflow_Lambda([this]{ return ControlsViewModel->GetCurrentCreationWorkflow(); });
+		default: break;
+		}
+	}
+	
+	return SNew(SSandboxListRow, InItem.ToSharedRef(), InOwnerTable)
+		.ColumnFactories(ColumnFactoryMap)
+		.Padding(Padding)
+		.HighlightText(HighlightTextAttr);
+}
+
+void SSandboxListView::OnItemDoubleClicked(TSharedPtr<FSandboxListItem> InItem)
+{
+	if (DummyItem && DummyItem->Item == InItem)
+	{
+		return;
+	}
+	
+	ControlsViewModel->LoadSandbox(InItem->SandboxInfo.SandboxRoot);
+}
+
+void SSandboxListView::OnItemsChanged()
+{
+	const TArray<TSharedPtr<FSandboxListItem>>& NewItems = ListViewModel->GetItems();
+	const bool bHasDummyItem = DummyItem.IsSet();
+	
+	Items.Empty(NewItems.Num() + static_cast<int32>(bHasDummyItem));
+	if (bHasDummyItem)
+	{
+		Items.Add(DummyItem->Item);
+	}
+	Items.Append(NewItems);
+	
+	ListView->RequestListRefresh();
+}
+
+void SSandboxListView::OnCreationWorkflowStarted(const FSandboxCreationWorkflow& InWorkflow)
+{
+	if (ensure(!DummyItem))
+	{
+		DummyItem.Emplace(MakeShared<FSandboxListItem>(FSandboxInfo{}), EDummySandboxItemType::CreateSandbox);
+		Items.Insert(DummyItem->Item, 0);
+		ListView->RequestListRefresh();
+	}
+}
+
+void SSandboxListView::OnCreationWorkflowEnded()
+{
+	if (ensure(DummyItem))
+	{
+		Items.RemoveSingle(DummyItem->Item);
+		DummyItem.Reset();
+		ListView->RequestListRefresh();
+	}
+}
+}
+
+#undef LOCTEXT_NAMESPACE
\ No newline at end of file
diff --git a/Engine/Plugins/Developer/Sandbox/SandboxedEditing/Source/SandboxedEditing/Private/Features/Browser/Widgets/Unsandboxed/List/SSandboxListView.h b/Engine/Plugins/Developer/Sandbox/SandboxedEditing/Source/SandboxedEditing/Private/Features/Browser/Widgets/Unsandboxed/List/SSandboxListView.h
new file mode 100644
--- /dev/null
+++ b/Engine/Plugins/Developer/Sandbox/SandboxedEditing/Source/SandboxedEditing/Private/Features/Browser/Widgets/Unsandboxed/List/SSandboxListView.h
@@ -0,0 +1,90 @@
+﻿// Copyright Epic Games, Inc. All Rights Reserved.
+
+#pragma once
+
+#include "DummySandboxItem.h"
+#include "Features/Browser/ViewModels/SandboxCreationWorkflow.h"
+#include "Columns/SandboxColumns.h"
+#include "Widgets/DeclarativeSyntaxSupport.h"
+#include "Widgets/SCompoundWidget.h"
+
+class SHeaderRow;
+class ITableRow;
+class STableViewBase;
+template <typename ItemType> class SListView;
+
+namespace UE::SandboxedEditing
+{
+class FSandboxControlsViewModel;
+class FSandboxListItem;
+class FSandboxListViewModel;
+enum class EBreakBehavior : uint8;
+
+/** A list view containing all the sandboxes. */
+class SSandboxListView : public SCompoundWidget
+{
+public:
+	
+	struct FViewModels
+	{
+		TSharedRef<FSandboxListViewModel> ListViewModel;
+		TSharedRef<FSandboxControlsViewModel> ControlsViewModel;
+	};
+	
+	SLATE_BEGIN_ARGS(SSandboxListView){}
+		/** The text to highlight in the columns. */
+		SLATE_ATTRIBUTE(FText, HighlightText)
+		/** Used to build the columns. */
+		SLATE_ARGUMENT(FSandboxColumnFactoryMap, ColumnFactories)
+	SLATE_END_ARGS()
+
+	void Construct(const FArguments& InArgs, const FViewModels& InViewModels);
+	
+	/** Enumerates every selected sandbox. */
+	void ForEachSelectedItem(TFunctionRef<EBreakBehavior(const FString& InRootPath)> InCallback) const;
+	
+	/** @return The sandbox that is selected if exactly one item is selected. */
+	TOptional<FString> GetSelectedItem() const;
+	
+	/** @return Number of items selected */
+	int32 NumSelected() const;
+	
+private:
+	
+	/** Knows about the items that are displayed in the browser. */
+	TSharedPtr<FSandboxListViewModel> ListViewModel;
+	/** Responds to creation of FSandboxCreationWorkflow by adding a dummy row into the table view. */
+	TSharedPtr<FSandboxControlsViewModel> ControlsViewModel;
+	
+	/** Used to build the columns. */
+	FSandboxColumnFactoryMap ColumnFactoryMap;
+	/** The text to highlight in the columns. */
+	TAttribute<FText> HighlightTextAttr;
+	
+	/** Displays the items. */
+	TSharedPtr<SListView<TSharedPtr<FSandboxListItem>>> ListView;
+	/** The items displayed in the ListView. A copy of the items in ListViewModel plus any dummy row. */
+	TArray<TSharedPtr<FSandboxListItem>> Items;
+	
+	/**
+	 * Dummy item used to make the ListView display an additional row, e.g. for creating a new sandbox.
+	 * Usually inserted at the beginning of items while in a creation workflow. It is inserted first so it's displayed on the top.
+	 */
+	TOptional<FDummySandboxItemInfo> DummyItem;
+	
+	/** @return The header row to display at the top of the table */
+	TSharedRef<SHeaderRow> MakeHeaderRow();
+	
+	/** @return Widget for the items */
+	TSharedRef<ITableRow> OnGenerateRow(TSharedPtr<FSandboxListItem> InItem, const TSharedRef<STableViewBase>& InOwnerTable);
+	/** Loads the clicked sandbox */
+	void OnItemDoubleClicked(TSharedPtr<FSandboxListItem> InItem);
+
+	/** Invoked when the list needs to be regenerated. */
+	void OnItemsChanged();
+	
+	void OnCreationWorkflowStarted(const FSandboxCreationWorkflow& InWorkflow);
+	void OnCreationWorkflowEnded();
+};
+}
+
diff --git a/Engine/Plugins/Developer/Sandbox/SandboxedEditing/Source/SandboxedEditing/Private/Features/Browser/Widgets/Unsandboxed/SBrowserControls.cpp b/Engine/Plugins/Developer/Sandbox/SandboxedEditing/Source/SandboxedEditing/Private/Features/Browser/Widgets/Unsandboxed/SBrowserControls.cpp
new file mode 100644
--- /dev/null
+++ b/Engine/Plugins/Developer/Sandbox/SandboxedEditing/Source/SandboxedEditing/Private/Features/Browser/Widgets/Unsandboxed/SBrowserControls.cpp
@@ -0,0 +1,58 @@
+﻿// Copyright Epic Games, Inc. All Rights Reserved.
+
+#include "SBrowserControls.h"
+
+#include "SandboxedEditingStyle.h"
+#include "Features/Browser/Commands/BrowserCommands.h"
+#include "Features/Browser/ViewModels/SandboxControlsViewModel.h"
+#include "Filters/CustomTextFilters.h"
+#include "Framework/Commands/GenericCommands.h"
+#include "Framework/MultiBox/MultiBoxBuilder.h"
+#include "Utils/WidgetUtils.h"
+#include "Widgets/SBoxPanel.h"
+
+#define LOCTEXT_NAMESPACE "SBrowserControls"
+
+namespace UE::SandboxedEditing
+{
+void SBrowserControls::Construct(const FArguments& InArgs, const TSharedRef<FSandboxControlsViewModel>& InControlsViewModel)
+{
+	ControlsViewModel = InControlsViewModel;
+	
+	ChildSlot
+	[
+		BuildButtonBar(InArgs)
+	];
+}
+
+TSharedRef<SWidget> SBrowserControls::BuildButtonBar(const FArguments& InArgs)
+{
+	FToolBarBuilder RowBuilder(nullptr, FMultiBoxCustomization::None);
+	
+	const TSharedPtr<FUICommandList>& CommandList = InArgs._CommandList;
+	FGenericCommands& GenericCommands = FGenericCommands::Get();
+	FBrowserCommands& Commands = FBrowserCommands::Get();
+	
+	RowBuilder.AddWidget(
+		MakeIconButtonByCommand(
+			FSandboxedEditingStyle::Get().GetBrush("SandboxedEditing.Browser.NewSandbox"),
+			CommandList,
+			Commands.CreateNewSandbox, LOCTEXT("CreateNew", "Create a new sandbox")
+		)
+		);
+	
+	RowBuilder.AddSeparator();
+	
+	RowBuilder.AddWidget(
+		MakeIconButtonByCommand(
+			FSandboxedEditingStyle::Get().GetBrush("SandboxedEditing.Browser.DeleteSandbox"),
+			CommandList,
+			GenericCommands.Delete, LOCTEXT("Delete", "Delete sandbox")
+		)
+		);
+	
+	return RowBuilder.MakeWidget();
+}
+}
+
+#undef LOCTEXT_NAMESPACE
\ No newline at end of file
diff --git a/Engine/Plugins/Developer/Sandbox/SandboxedEditing/Source/SandboxedEditing/Private/Features/Browser/Widgets/Unsandboxed/SBrowserControls.h b/Engine/Plugins/Developer/Sandbox/SandboxedEditing/Source/SandboxedEditing/Private/Features/Browser/Widgets/Unsandboxed/SBrowserControls.h
new file mode 100644
--- /dev/null
+++ b/Engine/Plugins/Developer/Sandbox/SandboxedEditing/Source/SandboxedEditing/Private/Features/Browser/Widgets/Unsandboxed/SBrowserControls.h
@@ -0,0 +1,38 @@
+﻿// Copyright Epic Games, Inc. All Rights Reserved.
+
+#pragma once
+
+#include "Framework/SlateDelegates.h"
+#include "Widgets/DeclarativeSyntaxSupport.h"
+#include "Widgets/SCompoundWidget.h"
+
+class FUICommandList;
+
+namespace UE::SandboxedEditing
+{
+class FSandboxControlsViewModel;
+class FSandboxListItem;
+class FSandboxListViewModel;
+
+/** Controls widget to place on the top section of the browser. */
+class SBrowserControls : public SCompoundWidget
+{
+public:
+
+	SLATE_BEGIN_ARGS(SBrowserControls){}
+		/** The command list to bind selection-based commands to. */
+		SLATE_ARGUMENT(TSharedPtr<FUICommandList>, CommandList)
+	SLATE_END_ARGS()
+
+	void Construct(const FArguments& InArgs, const TSharedRef<FSandboxControlsViewModel>& InControlsViewModel);
+	
+private:
+	
+	/** View model for mutating */
+	TSharedPtr<FSandboxControlsViewModel> ControlsViewModel;
+	
+	/** Builds the button bar */
+	TSharedRef<SWidget> BuildButtonBar(const FArguments& InArgs);
+};
+}
+
diff --git a/Engine/Plugins/Developer/Sandbox/SandboxedEditing/Source/SandboxedEditing/Private/Features/Browser/Widgets/Unsandboxed/SBrowserFilters.cpp b/Engine/Plugins/Developer/Sandbox/SandboxedEditing/Source/SandboxedEditing/Private/Features/Browser/Widgets/Unsandboxed/SBrowserFilters.cpp
new file mode 100644
--- /dev/null
+++ b/Engine/Plugins/Developer/Sandbox/SandboxedEditing/Source/SandboxedEditing/Private/Features/Browser/Widgets/Unsandboxed/SBrowserFilters.cpp
@@ -0,0 +1,41 @@
+﻿// Copyright Epic Games, Inc. All Rights Reserved.
+
+#include "SBrowserFilters.h"
+
+#include "Features/Browser/ViewModels/List/FilterSandboxViewModel.h"
+#include "Internationalization/Text.h"
+#include "Widgets/Input/SSearchBox.h"
+
+#define LOCTEXT_NAMESPACE "SSandboxFilters"
+
+namespace UE::SandboxedEditing
+{
+void SBrowserFilters::Construct(const FArguments& InArgs, const TSharedPtr<FFilterSandboxViewModel>& InFilterSandboxViewModel)
+{
+	FilterSandboxViewModel = InFilterSandboxViewModel;
+	
+	ChildSlot
+	[
+		SNew(SSearchBox)
+		.HintText(LOCTEXT("SearchHint", "Search Sandbox"))
+		.OnTextChanged(this, &SBrowserFilters::OnSearchTextChanged)
+		.OnTextCommitted(this, &SBrowserFilters::OnSearchTextCommitted)
+		.DelayChangeNotificationsWhileTyping(true)
+	];
+}
+
+void SBrowserFilters::OnSearchTextChanged(const FText& Text)
+{
+	FilterSandboxViewModel->SetSearchText(Text);
+}
+
+void SBrowserFilters::OnSearchTextCommitted(const FText& InFilterText, ETextCommit::Type)
+{
+	if (!InFilterText.EqualTo(FilterSandboxViewModel->GetSearchText()))
+	{
+		OnSearchTextChanged(InFilterText);
+	}
+}
+}
+
+#undef LOCTEXT_NAMESPACE
\ No newline at end of file
diff --git a/Engine/Plugins/Developer/Sandbox/SandboxedEditing/Source/SandboxedEditing/Private/Features/Browser/Widgets/Unsandboxed/SBrowserFilters.h b/Engine/Plugins/Developer/Sandbox/SandboxedEditing/Source/SandboxedEditing/Private/Features/Browser/Widgets/Unsandboxed/SBrowserFilters.h
new file mode 100644
--- /dev/null
+++ b/Engine/Plugins/Developer/Sandbox/SandboxedEditing/Source/SandboxedEditing/Private/Features/Browser/Widgets/Unsandboxed/SBrowserFilters.h
@@ -0,0 +1,30 @@
+﻿// Copyright Epic Games, Inc. All Rights Reserved.
+
+#pragma once
+
+#include "Widgets/DeclarativeSyntaxSupport.h"
+#include "Widgets/SCompoundWidget.h"
+
+namespace UE::SandboxedEditing
+{
+class FFilterSandboxViewModel;
+
+class SBrowserFilters : public SCompoundWidget
+{
+public:
+	
+	SLATE_BEGIN_ARGS(SBrowserFilters) {}
+	SLATE_END_ARGS()
+
+	void Construct(const FArguments& InArgs, const TSharedPtr<FFilterSandboxViewModel>& InFilterSandboxViewModel);
+	
+private:
+	
+	/** Knows how to filter sandbox items. */
+	TSharedPtr<FFilterSandboxViewModel> FilterSandboxViewModel;
+	
+	void OnSearchTextChanged(const FText& Text);
+	void OnSearchTextCommitted(const FText& InFilterText, ETextCommit::Type Arg);
+};
+}
+
diff --git a/Engine/Plugins/Developer/Sandbox/SandboxedEditing/Source/SandboxedEditing/Private/Features/Browser/Widgets/Unsandboxed/SBrowserSandboxDetails.cpp b/Engine/Plugins/Developer/Sandbox/SandboxedEditing/Source/SandboxedEditing/Private/Features/Browser/Widgets/Unsandboxed/SBrowserSandboxDetails.cpp
new file mode 100644
--- /dev/null
+++ b/Engine/Plugins/Developer/Sandbox/SandboxedEditing/Source/SandboxedEditing/Private/Features/Browser/Widgets/Unsandboxed/SBrowserSandboxDetails.cpp
@@ -0,0 +1,102 @@
+﻿// Copyright Epic Games, Inc. All Rights Reserved.
+
+#include "SBrowserSandboxDetails.h"
+
+#include "Features/Browser/Widgets/Shared/SSandboxDescription.h"
+#include "Internationalization/Text.h"
+#include "Widgets/Layout/SBox.h"
+#include "Widgets/Layout/SExpandableArea.h"
+#include "Widgets/Layout/SWidgetSwitcher.h"
+#include "Widgets/Text/STextBlock.h"
+
+#define LOCTEXT_NAMESPACE "SBrowserSandboxDetails"
+
+namespace UE::SandboxedEditing
+{
+namespace SandboxDetails
+{
+static TSharedRef<SWidget> MakeMessageWidget(TAttribute<FText> InMessageAttr)
+{
+	return SNew(SBox)
+		.HAlign(HAlign_Center)
+		.VAlign(VAlign_Center)
+		.Clipping(EWidgetClipping::ClipToBounds)
+		[
+			SNew(STextBlock)
+			.Text(InMessageAttr)
+			.Justification(ETextJustify::Center)
+		];
+}
+}
+void SBrowserSandboxDetails::Construct(const FArguments& InArgs, const TSharedRef<FSandboxMetaDataViewModel>& InMetaDataViewModel)
+{
+	NumSelectedAttr = InArgs._NumSelected;
+	ChildSlot
+	[
+		SNew(SExpandableArea)
+		.BodyBorderImage(FAppStyle::Get().GetBrush("ToolPanel.GroupBorder"))
+		.HeaderContent()
+		[
+			SNew(STextBlock)
+			.Text(LOCTEXT("Details.Label", "Details"))
+			.ToolTipText(LOCTEXT("Details.ToolTip", "Shows more info about the selected sandbox"))
+		]
+		.BodyContent() [ MakeWidgetSwitcherContent(InArgs, InMetaDataViewModel) ]
+	];
+}
+
+TSharedRef<SWidget> SBrowserSandboxDetails::MakeWidgetSwitcherContent(
+	const FArguments& InArgs, const TSharedRef<FSandboxMetaDataViewModel>& InMetaDataViewModel
+	)
+{
+	return SNew(SBox)
+		.MinDesiredHeight(120)
+		[
+			SNew(SWidgetSwitcher)
+			.WidgetIndex(this, &SBrowserSandboxDetails::GetWidgetIndex)
+
+			+SWidgetSwitcher::Slot()
+			[
+				SNew(SSandboxDescription, InMetaDataViewModel)
+				.SandboxPath(InArgs._SandboxPath)
+			]
+			
+			+SWidgetSwitcher::Slot()
+			[
+				SandboxDetails::MakeMessageWidget(LOCTEXT("NoneSelected", "Select a sandbox to view details."))
+			]
+
+			+SWidgetSwitcher::Slot()
+			[
+				SandboxDetails::MakeMessageWidget(LOCTEXT("TooManySelected", "Select only one sandbox to view details."))
+			]
+		];
+}
+
+int32 SBrowserSandboxDetails::GetWidgetIndex() const
+{
+	constexpr int32 MainContent = 0;
+	constexpr int32 NoneSelected = 1;
+	constexpr int32 TooManySelected = 2;
+	
+	const bool bIsSet = NumSelectedAttr.IsSet() || NumSelectedAttr.IsBound();
+	if (!bIsSet)
+	{
+		return MainContent;
+	}
+	
+	const int32 NumSelected = NumSelectedAttr.Get();
+	if (NumSelected == 1)
+	{
+		return MainContent;
+	}
+	
+	if (NumSelected == 0)
+	{
+		return NoneSelected;
+	}
+	return TooManySelected;
+}
+}
+
+#undef LOCTEXT_NAMESPACE
\ No newline at end of file
diff --git a/Engine/Plugins/Developer/Sandbox/SandboxedEditing/Source/SandboxedEditing/Private/Features/Browser/Widgets/Unsandboxed/SBrowserSandboxDetails.h b/Engine/Plugins/Developer/Sandbox/SandboxedEditing/Source/SandboxedEditing/Private/Features/Browser/Widgets/Unsandboxed/SBrowserSandboxDetails.h
new file mode 100644
--- /dev/null
+++ b/Engine/Plugins/Developer/Sandbox/SandboxedEditing/Source/SandboxedEditing/Private/Features/Browser/Widgets/Unsandboxed/SBrowserSandboxDetails.h
@@ -0,0 +1,37 @@
+﻿// Copyright Epic Games, Inc. All Rights Reserved.
+
+#pragma once
+
+#include "Widgets/DeclarativeSyntaxSupport.h"
+#include "Widgets/SCompoundWidget.h"
+
+namespace UE::SandboxedEditing
+{
+class FSandboxMetaDataViewModel;
+
+/** Shows details for a sandbox in the browser widget. Handles showing text such as "select sandbox to view details", etc. */
+class SBrowserSandboxDetails : public SCompoundWidget
+{
+public:
+	
+	SLATE_BEGIN_ARGS(SBrowserSandboxDetails){}
+		/** Path to the sandbox for which the description is being edited. Returning unset disables the edit box. */
+		SLATE_ATTRIBUTE(TOptional<FString>, SandboxPath)
+		/** The number of selected sandboxes. Relevant for the content displayed. */
+		SLATE_ATTRIBUTE(int32, NumSelected)
+	SLATE_END_ARGS()
+
+	void Construct(const FArguments& InArgs, const TSharedRef<FSandboxMetaDataViewModel>& InMetaDataViewModel);
+	
+private:
+	
+	/** Gets the number of selected items. Determines the content displayed. */
+	TAttribute<int32> NumSelectedAttr;
+	
+	/** @return The widget switcher that switches between the different messages and actual details content. */
+	TSharedRef<SWidget> MakeWidgetSwitcherContent(const FArguments& InArgs, const TSharedRef<FSandboxMetaDataViewModel>& InMetaDataViewModel);
+	
+	/** @return Index for widget switcher to display. Determined by how many items are selected. */
+	int32 GetWidgetIndex() const;
+};
+}
diff --git a/Engine/Plugins/Developer/Sandbox/SandboxedEditing/Source/SandboxedEditing/Private/Features/Browser/Widgets/Unsandboxed/SSandboxedEditingBrowser.cpp b/Engine/Plugins/Developer/Sandbox/SandboxedEditing/Source/SandboxedEditing/Private/Features/Browser/Widgets/Unsandboxed/SSandboxedEditingBrowser.cpp
new file mode 100644
--- /dev/null
+++ b/Engine/Plugins/Developer/Sandbox/SandboxedEditing/Source/SandboxedEditing/Private/Features/Browser/Widgets/Unsandboxed/SSandboxedEditingBrowser.cpp
@@ -0,0 +1,104 @@
+﻿// Copyright Epic Games, Inc. All Rights Reserved.
+
+#include "SSandboxedEditingBrowser.h"
+
+#include "Features/Browser/Commands/BrowserCommandBindings.h"
+#include "Framework/Commands/GenericCommands.h"
+#include "GenericBrowser/SSessionBrowserFrame.h"
+#include "List/SSandboxListView.h"
+#include "Misc/MessageDialog.h"
+#include "SBrowserControls.h"
+#include "SBrowserFilters.h"
+#include "SBrowserSandboxDetails.h"
+#include "Features/Browser/ViewModels/List/FilterSandboxViewModel.h"
+#include "Utils/BreakBehavior.h"
+#include "Widgets/Text/STextBlock.h"
+
+#define LOCTEXT_NAMESPACE "SSandboxedEditingBrowser"
+
+namespace UE::SandboxedEditing
+{
+void SSandboxedEditingBrowser::Construct(const FArguments& InArgs, const FViewModels& InViewModels)
+{
+	ControlsViewModel = InViewModels.ControlsViewModel;
+	CommandList = InArgs._CommandList;
+	
+	const SSandboxListView::FViewModels ListViewModels { InViewModels.ListViewModel, InViewModels.ControlsViewModel };
+	const TSharedRef<SSandboxListView> ListViewRef = SNew(SSandboxListView, ListViewModels)
+		.ColumnFactories(InArgs._ColumnFactories)
+		.HighlightText(InViewModels.FilterSandboxViewModel, &FFilterSandboxViewModel::GetSearchText);
+	ListView = ListViewRef;
+	
+	ChildSlot
+	[
+		SNew(SSessionBrowserFrame)
+		.ControlContent() [ SNew(SBrowserControls, InViewModels.ControlsViewModel).CommandList(InArgs._CommandList) ]
+		.SearchContent() [ SNew(SBrowserFilters, InViewModels.FilterSandboxViewModel)]
+		.SessionContent() [ ListViewRef ]
+		.DetailsContent()
+		[
+			SNew(SBrowserSandboxDetails, InViewModels.MetaDataViewModel)
+			.SandboxPath(ListViewRef, &SSandboxListView::GetSelectedItem)
+			.NumSelected(ListViewRef, &SSandboxListView::NumSelected)
+		]
+	];
+	
+	BindCommands();
+}
+
+void SSandboxedEditingBrowser::BindCommands()
+{
+	if (!CommandList)
+	{
+		return;
+	}
+	
+	FGenericCommands& GenericCommands = FGenericCommands::Get();
+	CommandList->MapAction(
+		GenericCommands.Delete, 
+		FExecuteAction::CreateSP(this, &SSandboxedEditingBrowser::HandleDeleteSelection), 
+		FCanExecuteAction::CreateSP(this, &SSandboxedEditingBrowser::CanDeleteSelection)
+		);
+}
+
+void SSandboxedEditingBrowser::HandleDeleteSelection() const
+{
+	if (CanDeleteSelection())
+	{
+		const EAppReturnType::Type Result = FMessageDialog::Open(
+			EAppMsgType::YesNo, 
+			FText::Format(LOCTEXT("Delete.MessageFmt", "Are you sure you want to delete {0} {0}|plural(one=sandbox,other=sandboxes)?"), ListView->NumSelected()),
+			FText::Format(LOCTEXT("Delete.TitleFmt", "Delete {0}|plural(one=sandbox,other=sandboxes)"), ListView->NumSelected())
+			);
+		if (Result != EAppReturnType::Yes)
+		{
+			return;
+		}
+		
+		ListView->ForEachSelectedItem([this](const FString& InRootPath)
+		{
+			ControlsViewModel->DeleteSandbox(InRootPath);
+			return EBreakBehavior::Continue;
+		});
+	}
+}
+
+bool SSandboxedEditingBrowser::CanDeleteSelection() const
+{
+	if (ListView->NumSelected() == 0)
+	{
+		return false;
+	}
+	
+	bool bCanDelete = true;
+	ListView->ForEachSelectedItem([this, &bCanDelete](const FString& InRootPath)
+	{
+		bCanDelete &= ControlsViewModel->CanDeleteSandbox(InRootPath);
+		return bCanDelete ? EBreakBehavior::Continue : EBreakBehavior::Break;
+	});
+	
+	return bCanDelete;
+}
+}
+
+#undef LOCTEXT_NAMESPACE
\ No newline at end of file
diff --git a/Engine/Plugins/Developer/Sandbox/SandboxedEditing/Source/SandboxedEditing/Private/Features/Browser/Widgets/Unsandboxed/SSandboxedEditingBrowser.h b/Engine/Plugins/Developer/Sandbox/SandboxedEditing/Source/SandboxedEditing/Private/Features/Browser/Widgets/Unsandboxed/SSandboxedEditingBrowser.h
new file mode 100644
--- /dev/null
+++ b/Engine/Plugins/Developer/Sandbox/SandboxedEditing/Source/SandboxedEditing/Private/Features/Browser/Widgets/Unsandboxed/SSandboxedEditingBrowser.h
@@ -0,0 +1,62 @@
+﻿// Copyright Epic Games, Inc. All Rights Reserved.
+
+#pragma once
+
+#include "List/Columns/SandboxColumns.h"
+#include "Widgets/DeclarativeSyntaxSupport.h"
+#include "Widgets/SCompoundWidget.h"
+
+class FUICommandList;
+
+namespace UE::SandboxedEditing
+{
+class FBrowserCommandBindings;
+class FFilterSandboxViewModel;
+class FSandboxControlsViewModel;
+class FSandboxListViewModel;
+class FSandboxMetaDataViewModel;
+class SSandboxListView;
+
+/** Displays a list of sandboxes. Shows content only */
+class SSandboxedEditingBrowser : public SCompoundWidget
+{
+public:
+	
+	struct FViewModels
+	{
+		TSharedRef<FFilterSandboxViewModel> FilterSandboxViewModel;
+		TSharedRef<FSandboxListViewModel> ListViewModel;
+		TSharedRef<FSandboxControlsViewModel> ControlsViewModel;
+		TSharedRef<FSandboxMetaDataViewModel> MetaDataViewModel;
+	};
+
+	SLATE_BEGIN_ARGS(SSandboxedEditingBrowser){}
+		/** The command list to bind selection-based commands to. */
+		SLATE_ARGUMENT(TSharedPtr<FUICommandList>, CommandList)
+		/** Used to build the columns. */
+		SLATE_ARGUMENT(FSandboxColumnFactoryMap, ColumnFactories)
+	SLATE_END_ARGS()
+
+	void Construct(const FArguments& InArgs, const FViewModels& InViewModels);
+	
+private:
+	
+	/** Used for implementing the commands we bind to. */
+	TSharedPtr<FSandboxControlsViewModel> ControlsViewModel;
+	/** Optional command list for selection based commands, such as delete. */
+	TSharedPtr<FUICommandList> CommandList;
+	
+	/** List view displaying the sandboxes. */
+	TSharedPtr<SSandboxListView> ListView;
+	
+	
+	/** Binds the commands only possible with UI, e.g. those that perform an action on the selection. */
+	void BindCommands();
+	
+	/** Handles deleting the selection. */
+	void HandleDeleteSelection() const;
+	/** @return Whether the current selection can be deleted. */
+	bool CanDeleteSelection() const;
+};
+}
+
diff --git a/Engine/Plugins/Developer/Sandbox/SandboxedEditing/Source/SandboxedEditing/Private/Features/StatusBar/Commands/StatusBarCommands.cpp b/Engine/Plugins/Developer/Sandbox/SandboxedEditing/Source/SandboxedEditing/Private/Features/StatusBar/Commands/StatusBarCommands.cpp
--- a/Engine/Plugins/Developer/Sandbox/SandboxedEditing/Source/SandboxedEditing/Private/Features/StatusBar/Commands/StatusBarCommands.cpp
+++ b/Engine/Plugins/Developer/Sandbox/SandboxedEditing/Source/SandboxedEditing/Private/Features/StatusBar/Commands/StatusBarCommands.cpp
@@ -13,7 +13,7 @@ FStatusBarCommands::FStatusBarCommands()
 	(
 		"SandboxedEditing.StatusBar",
 		NSLOCTEXT("Contexts", "SandboxedEditing", "SandboxedEditing"),
-		"LevelEditor",
+		NAME_None,
 		FAppStyle::GetAppStyleSetName()
 	)
 {}
diff --git a/Engine/Plugins/Developer/Sandbox/SandboxedEditing/Source/SandboxedEditing/Private/Features/StatusBar/StatusBarFeature.cpp b/Engine/Plugins/Developer/Sandbox/SandboxedEditing/Source/SandboxedEditing/Private/Features/StatusBar/StatusBarFeature.cpp
--- a/Engine/Plugins/Developer/Sandbox/SandboxedEditing/Source/SandboxedEditing/Private/Features/StatusBar/StatusBarFeature.cpp
+++ b/Engine/Plugins/Developer/Sandbox/SandboxedEditing/Source/SandboxedEditing/Private/Features/StatusBar/StatusBarFeature.cpp
@@ -4,7 +4,6 @@
 
 #include "ToolMenus.h"
 #include "Features/StatusBar/Widgets/SSandboxStatusBarWidget.h"
-#include "Framework/Application/SlateApplication.h"
 
 namespace UE::SandboxedEditing
 {
@@ -17,11 +16,6 @@ static TSharedRef<SWidget> CreateStatusBarWidget(const TSharedRef<FSandboxSystem
 
 static void RegisterStatusBar(const TSharedRef<FSandboxSystemModel>& InViewModel, const FStatusBarCommandMappings& InCommandMappings)
 {
-	if (!FSlateApplication::IsInitialized())
-	{
-		return;
-	}
-	
 	UToolMenu* Menu = UToolMenus::Get()->ExtendMenu(TEXT("LevelEditor.StatusBar.ToolBar"));
 	FToolMenuSection& AssetManagementSection = Menu->AddSection(
 		TEXT("Asset Management"), FText::GetEmpty(), FToolMenuInsert(NAME_None, EToolMenuInsertType::First)
diff --git a/Engine/Plugins/Developer/Sandbox/SandboxedEditing/Source/SandboxedEditing/Private/Features/StatusBar/Widgets/SSandboxStatusBarWidget.cpp b/Engine/Plugins/Developer/Sandbox/SandboxedEditing/Source/SandboxedEditing/Private/Features/StatusBar/Widgets/SSandboxStatusBarWidget.cpp
--- a/Engine/Plugins/Developer/Sandbox/SandboxedEditing/Source/SandboxedEditing/Private/Features/StatusBar/Widgets/SSandboxStatusBarWidget.cpp
+++ b/Engine/Plugins/Developer/Sandbox/SandboxedEditing/Source/SandboxedEditing/Private/Features/StatusBar/Widgets/SSandboxStatusBarWidget.cpp
@@ -3,7 +3,7 @@
 #include "SSandboxStatusBarWidget.h"
 
 #include "Features/StatusBar/Commands/StatusBarCommands.h"
-#include "Framework/Models/SandboxItem.h"
+#include "Framework/Models/SandboxInfo.h"
 #include "Framework/Models/SandboxSystemModel.h"
 #include "Framework/Notifications.h"
 #include "SActionButton.h"
@@ -86,10 +86,12 @@ void SSandboxStatusBarWidget::BuildSandboxesMenu(UToolMenu* InMenuBuilder)
 	Section.AddDynamicEntry("Existing Sandboxes",
 		FNewToolMenuSectionDelegate::CreateLambda([this](FToolMenuSection& InSection)
 	{
-		const TArray<FSandboxItem> Sandboxes = ViewModel->GetKnownSandboxes();
+		TArray<FSandboxInfo> Sandboxes = ViewModel->GetKnownSandboxes();
+		// Display them in alphabetical order
+		Sandboxes.Sort();
 
 		TArray<FString> DuplicateNameDetection;
-		for (const FSandboxItem& Sandbox : Sandboxes)
+		for (const FSandboxInfo& Sandbox : Sandboxes)
 		{
 			// Tool menus want unique names. It could happen that the user goes into the file system and manually edits the name entry in the JSON.
 			const FName MenuName(
diff --git a/Engine/Plugins/Developer/Sandbox/SandboxedEditing/Source/SandboxedEditing/Private/Framework/Models/SandboxInfo.h b/Engine/Plugins/Developer/Sandbox/SandboxedEditing/Source/SandboxedEditing/Private/Framework/Models/SandboxInfo.h
new file mode 100644
--- /dev/null
+++ b/Engine/Plugins/Developer/Sandbox/SandboxedEditing/Source/SandboxedEditing/Private/Framework/Models/SandboxInfo.h
@@ -0,0 +1,46 @@
+﻿// Copyright Epic Games, Inc. All Rights Reserved.
+
+#pragma once
+
+#include "Containers/UnrealString.h"
+#include "Data/VersionInfo.h"
+#include "Misc/DateTime.h"
+
+namespace UE::SandboxedEditing
+{
+/** Describes a sandbox. */
+struct FSandboxInfo
+{
+	/** Name the user has given the sandbox. May not be unique (but usually is). */
+	FString Name;
+
+	/** Description the user has given the sandbox. */
+	FString Description;
+
+	/** Path to the root of the sandbox directory. Unique for every sandbox item. */
+	FString SandboxRoot;
+	
+	/** Info about when the sandbox was last modified. */
+	FDateTime LastModified;
+	
+	/** The version of the engine when this sandbox was created in. */
+	FFileSandboxCore_EngineVersionInfo EngineVersion;
+
+	FSandboxInfo() = default;
+	explicit FSandboxInfo(
+		FString InName, FString InDescription, FString InDirectoryPath,
+		const FDateTime& InLastModified, FFileSandboxCore_EngineVersionInfo InEngineVersion
+		)
+		: Name(MoveTemp(InName))
+		, Description(MoveTemp(InDescription))
+		, SandboxRoot(MoveTemp(InDirectoryPath))
+		, LastModified(InLastModified)
+		, EngineVersion(MoveTemp(InEngineVersion))
+	{}
+
+	friend bool operator<(const FSandboxInfo& Left, const FSandboxInfo& Right)
+	{
+		return Left.Name < Right.Name;
+	}
+};
+}
diff --git a/Engine/Plugins/Developer/Sandbox/SandboxedEditing/Source/SandboxedEditing/Private/Framework/Models/SandboxItem.h b/Engine/Plugins/Developer/Sandbox/SandboxedEditing/Source/SandboxedEditing/Private/Framework/Models/SandboxItem.h
deleted file mode 100644
diff --git a/Engine/Plugins/Developer/Sandbox/SandboxedEditing/Source/SandboxedEditing/Private/Framework/Models/SandboxSystemModel.cpp b/Engine/Plugins/Developer/Sandbox/SandboxedEditing/Source/SandboxedEditing/Private/Framework/Models/SandboxSystemModel.cpp
--- a/Engine/Plugins/Developer/Sandbox/SandboxedEditing/Source/SandboxedEditing/Private/Framework/Models/SandboxSystemModel.cpp
+++ b/Engine/Plugins/Developer/Sandbox/SandboxedEditing/Source/SandboxedEditing/Private/Framework/Models/SandboxSystemModel.cpp
@@ -2,15 +2,20 @@
 
 #include "SandboxSystemModel.h"
 
+#include "Containers/UnrealString.h"
 #include "Data/SandboxMetaData.h"
 #include "Framework/Models/CreateSandboxArgs.h"
-#include "Framework/Models/SandboxItem.h"
+#include "Framework/Models/SandboxInfo.h"
 #include "IFileSandboxCoreModule.h"
 #include "ISandboxInstance.h"
 #include "ISandboxManager.h"
+#include "ISandboxRepository.h"
 #include "Misc/Paths.h"
 #include "Templates/Function.h"
 #include "Types/EBreakBehavior.h"
+#include "Types/SandboxMetaInfo.h"
+#include "Types/Manager/DeleteSandboxByDirectoryArgs.h"
+#include "Types/Manager/DeleteSandboxResult.h"
 #include "Types/Manager/LoadSandboxByDirectoryArgs.h"
 #include "Types/Manager/NewSandboxArgs.h"
 
@@ -18,6 +23,38 @@
 
 namespace UE::SandboxedEditing
 {
+FSandboxSystemModel::FSandboxSystemModel()
+{
+	using namespace FileSandboxCore;
+	IFileSandboxCoreModule& Module = IFileSandboxCoreModule::Get();
+	
+	ISandboxManager& Manager = Module.GetSandboxManager();
+	Manager.OnPostSandboxStartup().AddRaw(this, &FSandboxSystemModel::BroadcastOnLoadSandbox);
+	Manager.OnPostSandboxShutdown().AddRaw(this, &FSandboxSystemModel::BroadcastOnLeaveSandbox);
+	
+	ISandboxRepository& Repository = Module.GetDefaultSandboxRepository();
+	Repository.OnSandboxesChanged().AddRaw(this, &FSandboxSystemModel::OnSandboxesChanged);
+	Repository.OnSandboxMetaDataChanged().AddRaw(this, &FSandboxSystemModel::OnMetaDataChanged);
+}
+
+FSandboxSystemModel::~FSandboxSystemModel()
+{
+	using namespace FileSandboxCore;
+	
+	if (IFileSandboxCoreModule::IsAvailable())
+	{
+		IFileSandboxCoreModule& Module = IFileSandboxCoreModule::Get();
+		
+		ISandboxManager& Manager = Module.GetSandboxManager();
+		Manager.OnPostSandboxStartup().RemoveAll(this);
+		Manager.OnPostSandboxShutdown().RemoveAll(this);
+		
+		ISandboxRepository& Repository = Module.GetDefaultSandboxRepository();
+		Repository.OnSandboxesChanged().RemoveAll(this);
+		Repository.OnSandboxMetaDataChanged().RemoveAll(this);
+	}
+}
+
 bool FSandboxSystemModel::CreateNewSandbox(const FCreateSandboxArgs& InArgs)
 {
 	if (!CanCreateNewSandbox(InArgs.Name))
@@ -42,12 +79,13 @@ bool FSandboxSystemModel::CanCreateNewSandbox(const FString& InName, FText* Reas
 {
 	FileSandboxCore::EDirectorySuitability Suitability = FileSandboxCore::DetermineDirectorySuitability(InName);
 
-	if (Reason)
+	const bool bIsSuitable = Suitability == FileSandboxCore::EDirectorySuitability::Suitable;
+	if (Reason && !bIsSuitable)
 	{
 		*Reason = FileSandboxCore::FormatDirectorySuitabilityAsReason(Suitability, InName);
 	}
 	
-	return Suitability == FileSandboxCore::EDirectorySuitability::Suitable;
+	return bIsSuitable;
 }
 
 bool FSandboxSystemModel::LoadSandbox(const FString& InRootDirectory)
@@ -92,6 +130,36 @@ bool FSandboxSystemModel::CanLeaveSandbox() const
 	return SandboxManager.GetActiveSandboxInstance() != nullptr;
 }
 
+void FSandboxSystemModel::DeleteSandbox(const FString& InSandboxRoot) const
+{
+	if (CanDeleteSandbox(InSandboxRoot))
+	{
+		using namespace FileSandboxCore;
+		IFileSandboxCoreModule& Module = IFileSandboxCoreModule::Get();
+		ISandboxManager& SandboxManager = Module.GetSandboxManager();
+		SandboxManager.DeleteSandbox(FDeleteSandboxByDirectoryArgs{ InSandboxRoot });
+	}
+}
+
+#define SET_REASON(ReasonText) if (OutReason){ *OutReason = ReasonText; }
+bool FSandboxSystemModel::CanDeleteSandbox(const FString& InSandboxRoot, FText* OutReason) const
+{
+	if (IsActiveSandbox(InSandboxRoot))
+	{
+		SET_REASON(LOCTEXT("Delete.ActiveSandbox", "Cannot delete active sandbox"));
+		return false;
+	}
+	
+	if (!FileSandboxCore::IsRootSandboxDirectory(InSandboxRoot))
+	{
+		SET_REASON(LOCTEXT("Delete.NotSandbox", "Directory does not contain any sandbox"));
+		return false;
+	}
+	
+	return true;
+}
+#undef SET_REASON
+
 void FSandboxSystemModel::PersistAllChanges() const
 {
 	using namespace FileSandboxCore;
@@ -128,6 +196,14 @@ bool FSandboxSystemModel::CanDiscardAllChanges() const
 	return HasActiveSandbox();
 }
 
+FString FSandboxSystemModel::GetActiveSandboxPath() const
+{
+	using namespace FileSandboxCore;
+	IFileSandboxCoreModule& Module = IFileSandboxCoreModule::Get();
+	ISandboxInstance* SandboxInstance = Module.GetSandboxManager().GetActiveSandboxInstance();
+	return SandboxInstance ? SandboxInstance->GetRootDirectory() : FString();
+}
+
 FString FSandboxSystemModel::GetActiveSandboxName() const
 {
 	using namespace FileSandboxCore;
@@ -143,26 +219,53 @@ bool FSandboxSystemModel::HasActiveSandbox() const
 	return Module.GetSandboxManager().GetActiveSandboxInstance() != nullptr;
 }
 
-TArray<FSandboxItem> FSandboxSystemModel::GetKnownSandboxes() const
+void FSandboxSystemModel::ForEachSandbox(TFunctionRef<void(const FSandboxInfo& InRoot)> InProcess) const
 {
 	using namespace FileSandboxCore;
+	IFileSandboxCoreModule& Module = IFileSandboxCoreModule::Get();
 	
-	// This is very inefficient because it loads up every single manifest file, parses the JSON, etc.
-	// We should add a cache to FileSandboxCore for this.
-	TArray<FSandboxItem> SandboxDirectories;
-	ForEachSandbox([&SandboxDirectories](const FString& Directory)
+	ISandboxRepository& Repository = Module.GetDefaultSandboxRepository();
+	Repository.ForEachSandbox([&InProcess](const FString& InRootPath, const FSandboxMetaInfo& MetaInfo)
 	{
-		if (TOptional<FFileSandboxCore_SandboxMetaData> MetaData = LoadMetaData(Directory))
-		{
-			SandboxDirectories.Emplace(MoveTemp(MetaData->Name), MoveTemp(MetaData->Description), Directory);
-		}
-		return EBreakBehavior::Continue;
+		const FFileSandboxCore_SandboxMetaData MetaData = MetaInfo.UserMetaData;
+		InProcess(FSandboxInfo(MetaData.Name, MetaData.Description, InRootPath, GetLastModified(InRootPath), MetaInfo.VersionInfo.EngineVersion));
+		return FileSandboxCore::EBreakBehavior::Continue;
 	});
+}
+
+TArray<FSandboxInfo> FSandboxSystemModel::GetKnownSandboxes() const
+{
+	TArray<FSandboxInfo> Result;
+	ForEachSandbox([&Result](const FSandboxInfo& Info) { Result.Add(Info); });
+	return Result;
+}
 
-	// Display them in alphabetical order
-	SandboxDirectories.Sort();
+TOptional<FSandboxInfo> FSandboxSystemModel::GetSandboxInfo(const FString& InSandboxRoot) const
+{
+	using namespace FileSandboxCore;
+	IFileSandboxCoreModule& Module = IFileSandboxCoreModule::Get();
+	ISandboxRepository& Repository = Module.GetDefaultSandboxRepository();
 	
-	return SandboxDirectories;
+	TOptional<FSandboxInfo> Result;
+	Repository.ReadMetaData(InSandboxRoot, [&InSandboxRoot, &Result](const FSandboxMetaInfo& MetaInfo)
+	{
+		const FFileSandboxCore_SandboxMetaData MetaData = MetaInfo.UserMetaData;
+		Result.Emplace(MetaData.Name, MetaData.Description, InSandboxRoot, MetaInfo.LastModified, MetaInfo.VersionInfo.EngineVersion);
+	});
+	return Result;
+}
+
+void FSandboxSystemModel::SetDescription(const FString& InSandboxRoot, const FString& InDescription) const
+{
+	using namespace FileSandboxCore;
+	IFileSandboxCoreModule& Module = IFileSandboxCoreModule::Get();
+	ISandboxRepository& Repository = Module.GetDefaultSandboxRepository();
+	Repository.ReadMetaData(InSandboxRoot, [&InSandboxRoot, &InDescription](const FSandboxMetaInfo& MetaInfo)
+	{
+		FFileSandboxCore_SandboxMetaData MetaData = MetaInfo.UserMetaData;
+		MetaData.Description = InDescription;
+		SaveMetaData(MetaData, InSandboxRoot);
+	});
 }
 }
 
diff --git a/Engine/Plugins/Developer/Sandbox/SandboxedEditing/Source/SandboxedEditing/Private/Framework/Models/SandboxSystemModel.h b/Engine/Plugins/Developer/Sandbox/SandboxedEditing/Source/SandboxedEditing/Private/Framework/Models/SandboxSystemModel.h
--- a/Engine/Plugins/Developer/Sandbox/SandboxedEditing/Source/SandboxedEditing/Private/Framework/Models/SandboxSystemModel.h
+++ b/Engine/Plugins/Developer/Sandbox/SandboxedEditing/Source/SandboxedEditing/Private/Framework/Models/SandboxSystemModel.h
@@ -4,13 +4,19 @@
 
 #include "Delegates/Delegate.h"
 #include "Templates/UnrealTemplate.h"
+#include "Types/RepositoryChangedEvent.h"
+
+namespace UE::FileSandboxCore
+{
+class ISandboxInstance;
+}
 
 template<typename OptionalType> struct TOptional;
 
 namespace UE::SandboxedEditing
 {
 struct FCreateSandboxArgs;
-struct FSandboxItem;
+struct FSandboxInfo;
 	
 /**
  * View model for creating, loading, and querying the active sandbox.
@@ -21,6 +27,8 @@ struct FSandboxItem;
 class FSandboxSystemModel : public FNoncopyable
 {
 public:
+	FSandboxSystemModel();
+	~FSandboxSystemModel();
 
 	/** Creates a new sandbox with the given name. */
 	bool CreateNewSandbox(const FCreateSandboxArgs& InArgs);
@@ -38,6 +46,11 @@ public:
 	void LeaveSandbox() const;
 	/** @return Whether the sandbox can be left (false if there is no active one). */
 	bool CanLeaveSandbox() const;
+	
+	/** Deletes the sandbox. */
+	void DeleteSandbox(const FString& InSandboxRoot) const;
+	/** @return Whether the sandbox can be deleted. */
+	bool CanDeleteSandbox(const FString& InSandboxRoot, FText* OutReason = nullptr) const;
 
 	/** Persists all changes */
 	void PersistAllChanges() const;
@@ -49,14 +62,46 @@ public:
 	/** @return Whether all changes can be discarded */
 	bool CanDiscardAllChanges() const; 
 
+	/** @return Path to the root directory of the active sandbox. Empty if none is active. */
+	FString GetActiveSandboxPath() const;
 	/** @return Name of the currently active sandbox. Empty if none active. */
 	FString GetActiveSandboxName() const;
-
 	/** @return Whether there is an active sandbox */
 	bool HasActiveSandbox() const;
 
+	/** Enumerates all known sandboxes. */
+	void ForEachSandbox(TFunctionRef<void(const FSandboxInfo& InRoot)> InProcess) const;
 	/** @return The sandboxes we know of. */
-	TArray<FSandboxItem> GetKnownSandboxes() const;
+	TArray<FSandboxInfo> GetKnownSandboxes() const;
+	
+	/** @return Info about the sandbox. Unset if sandbox does not exist. */
+	TOptional<FSandboxInfo> GetSandboxInfo(const FString& InSandboxRoot) const;
+	/** Sets the sandbox description. */
+	void SetDescription(const FString& InSandboxRoot, const FString& InDescription) const;
+	
+	/** Invoked when the result of GetKnownSandboxes may have changed. */
+	FSimpleMulticastDelegate& OnKnownSandboxesChanged() { return OnKnownSandboxesChangedDelegate; }
+	
+	/** Invoked when the a new sandbox is loaded. */
+	FSimpleMulticastDelegate& OnLoadSandbox() { return OnLoadSandboxDelegate; }
+	/** Invoked when the current sandbox is left. */
+	FSimpleMulticastDelegate& OnLeaveSandbox() { return OnLeaveSandboxDelegate; }
+		
+private:
+	
+	/** Invoked when the result of GetKnownSandboxes may have changed. */
+	FSimpleMulticastDelegate OnKnownSandboxesChangedDelegate;
+	
+	/** Invoked when the a new sandbox is loaded. */
+	FSimpleMulticastDelegate OnLoadSandboxDelegate;
+	/** Invoked when the current sandbox is left. */
+	FSimpleMulticastDelegate OnLeaveSandboxDelegate;
+	
+	void OnSandboxesChanged(const FileSandboxCore::FRepositoryChangedEvent&) const { OnKnownSandboxesChangedDelegate.Broadcast(); }
+	void OnMetaDataChanged(const FString&) const { OnKnownSandboxesChangedDelegate.Broadcast(); }
+	
+	void BroadcastOnLoadSandbox(FileSandboxCore::ISandboxInstance&) { OnLoadSandboxDelegate.Broadcast(); }
+	void BroadcastOnLeaveSandbox() { OnLeaveSandboxDelegate.Broadcast(); }
 };
 }
 
diff --git a/Engine/Plugins/Developer/Sandbox/SandboxedEditing/Source/SandboxedEditing/Private/Framework/SandboxedEditingApp.cpp b/Engine/Plugins/Developer/Sandbox/SandboxedEditing/Source/SandboxedEditing/Private/Framework/SandboxedEditingApp.cpp
--- a/Engine/Plugins/Developer/Sandbox/SandboxedEditing/Source/SandboxedEditing/Private/Framework/SandboxedEditingApp.cpp
+++ b/Engine/Plugins/Developer/Sandbox/SandboxedEditing/Source/SandboxedEditing/Private/Framework/SandboxedEditingApp.cpp
@@ -7,7 +7,8 @@
 namespace UE::SandboxedEditing
 {
 FSandboxedEditingApp::FSandboxedEditingApp()
-	: SandboxViewModel(MakeShared<FSandboxSystemModel>())
-	, StatusBarFeature(SandboxViewModel)
+	: SandboxModel(MakeShared<FSandboxSystemModel>())
+	, StatusBarFeature(SandboxModel)
+	, BrowserFeature(SandboxModel)
 {}
 }
diff --git a/Engine/Plugins/Developer/Sandbox/SandboxedEditing/Source/SandboxedEditing/Private/Framework/SandboxedEditingApp.h b/Engine/Plugins/Developer/Sandbox/SandboxedEditing/Source/SandboxedEditing/Private/Framework/SandboxedEditingApp.h
--- a/Engine/Plugins/Developer/Sandbox/SandboxedEditing/Source/SandboxedEditing/Private/Framework/SandboxedEditingApp.h
+++ b/Engine/Plugins/Developer/Sandbox/SandboxedEditing/Source/SandboxedEditing/Private/Framework/SandboxedEditingApp.h
@@ -2,6 +2,7 @@
 
 #pragma once
 
+#include "Features/Browser/BrowserFeature.h"
 #include "Features/StatusBar/StatusBarFeature.h"
 #include "Templates/SharedPointer.h"
 #include "Templates/UnrealTemplate.h"
@@ -26,10 +27,12 @@ private:
 	 * View model for interacting with sandbox system.
 	 * Shared pointer in case the UI to avoid crash in case that though some bug the UI outlives this app for whatever reason.
 	 */
-	const TSharedRef<FSandboxSystemModel> SandboxViewModel;
+	const TSharedRef<FSandboxSystemModel> SandboxModel;
 	
 	/** Extends the status bar for Sandboxed Editing. */
 	const FStatusBarFeature StatusBarFeature;
+	/** Extends the editor with a nomad tab containing a browser for sandboxes. */
+	const FBrowserFeature BrowserFeature;
 };
 }
 
diff --git a/Engine/Plugins/Developer/Sandbox/SandboxedEditing/Source/SandboxedEditing/Private/SandboxEditingModule.cpp b/Engine/Plugins/Developer/Sandbox/SandboxedEditing/Source/SandboxedEditing/Private/SandboxEditingModule.cpp
--- a/Engine/Plugins/Developer/Sandbox/SandboxedEditing/Source/SandboxedEditing/Private/SandboxEditingModule.cpp
+++ b/Engine/Plugins/Developer/Sandbox/SandboxedEditing/Source/SandboxedEditing/Private/SandboxEditingModule.cpp
@@ -2,6 +2,9 @@
 
 #include "SandboxEditingModule.h"
 
+#include "CoreGlobals.h"
+#include "Framework/Application/SlateApplication.h"
+#include "Misc/CoreDelegates.h"
 #include "SandboxedEditingStyle.h"
 
 namespace UE::SandboxedEditing
@@ -9,11 +12,19 @@ namespace UE::SandboxedEditing
 void FSandboxEditingModule::StartupModule()
 {
 	FSandboxedEditingStyle::Get();
+	
+	FCoreDelegates::OnPostEngineInit.AddLambda([this]
+	{
+		if (FSlateApplication::IsInitialized() && GIsEditor)
+		{
+			EditingApp = MakeUnique<FSandboxedEditingApp>();
+		}
+	});
 }
 
 void FSandboxEditingModule::ShutdownModule()
 {
-	
+	EditingApp.Reset();
 }
 }
 
diff --git a/Engine/Plugins/Developer/Sandbox/SandboxedEditing/Source/SandboxedEditing/Private/SandboxEditingModule.h b/Engine/Plugins/Developer/Sandbox/SandboxedEditing/Source/SandboxedEditing/Private/SandboxEditingModule.h
--- a/Engine/Plugins/Developer/Sandbox/SandboxedEditing/Source/SandboxedEditing/Private/SandboxEditingModule.h
+++ b/Engine/Plugins/Developer/Sandbox/SandboxedEditing/Source/SandboxedEditing/Private/SandboxEditingModule.h
@@ -19,7 +19,7 @@ public:
 private:
 
 	/** Manages all the business logic and creates the editor workflows. */
-	FSandboxedEditingApp EditingApp;
+	TUniquePtr<FSandboxedEditingApp> EditingApp;
 };
 }
 
diff --git a/Engine/Plugins/Developer/Sandbox/SandboxedEditing/Source/SandboxedEditing/Private/SandboxedEditingStyle.cpp b/Engine/Plugins/Developer/Sandbox/SandboxedEditing/Source/SandboxedEditing/Private/SandboxedEditingStyle.cpp
--- a/Engine/Plugins/Developer/Sandbox/SandboxedEditing/Source/SandboxedEditing/Private/SandboxedEditingStyle.cpp
+++ b/Engine/Plugins/Developer/Sandbox/SandboxedEditing/Source/SandboxedEditing/Private/SandboxedEditingStyle.cpp
@@ -2,6 +2,7 @@
 
 #include "SandboxedEditingStyle.h"
 
+#include "Styling/SlateStyleMacros.h"
 #include "Styling/SlateStyleRegistry.h"
 
 namespace UE::SandboxedEditing
@@ -11,16 +12,17 @@ FSandboxedEditingStyle& FSandboxedEditingStyle::Get()
 	static FSandboxedEditingStyle Style;
 	return Style;
 }
-	
+
 FSandboxedEditingStyle::FSandboxedEditingStyle()
 	: FSlateStyleSet("SandboxedEditingStyle")
 {
-	const FString PluginContentDir = FPaths::EnginePluginsDir() / TEXT("Developer") / TEXT("Sandbox") / TEXT("SandboxedEditing") / TEXT("Resources");
+	const FString PluginContentDir = FPaths::EnginePluginsDir() / TEXT("Developer") / TEXT("Sandbox") / TEXT("SandboxedEditing") / TEXT("Content");
 	const FString EngineEditorSlateDir = FPaths::EngineContentDir() / TEXT("Editor") / TEXT("Slate");
 	SetContentRoot(PluginContentDir);
 	SetCoreContentRoot(EngineEditorSlateDir);
 	
-	Set("SandboxedEditing.Icons.CreateNewSandbox", new FSlateColorBrush(FLinearColor::Red));
+	RegisterBrowserStyle();
+	RegisterSharedStyle();
 	
 	FSlateStyleRegistry::RegisterSlateStyle(*this);
 }
@@ -29,4 +31,29 @@ FSandboxedEditingStyle::~FSandboxedEditingStyle()
 {
 	FSlateStyleRegistry::UnRegisterSlateStyle(*this);
 }
+
+void FSandboxedEditingStyle::RegisterBrowserStyle()
+{
+	const FLinearColor IconColorAndOpacity(FLinearColor(1.f, 1.f, 1.f, 1.f));
+	const FVector2D Icon24x24(24.0f, 24.0f);
+	
+	Set("SandboxedEditing.Browser.Icon", new IMAGE_BRUSH_SVG("SandboxBrowserIcon_24", Icon24x24, IconColorAndOpacity));
+	Set("SandboxedEditing.Browser.NewSandbox", new IMAGE_BRUSH_SVG("NewSandbox_24", Icon24x24, IconColorAndOpacity));
+	Set("SandboxedEditing.Browser.DeleteSandbox", new IMAGE_BRUSH("DeleteSandbox_48", Icon24x24, IconColorAndOpacity));
+	Set("SandboxedEditing.Browser.LeaveSandbox", new IMAGE_BRUSH("LeaveSandbox_48", Icon24x24, IconColorAndOpacity));
+	
+	Set("SandboxedEditing.Browser.RowPadding", FMargin(0.f, 4.f));
+	Set("SandboxedEditing.Browser.NameColumn.Padding", FMargin(16.f, 0.f, 0.f, 0.f));
+	Set("SandboxedEditing.Browser.NameColumn.FillWidth", 2.f);
+	Set("SandboxedEditing.Browser.VersionColumn.Padding", FMargin(8.f, 0.f, 0.f, 0.f));
+	Set("SandboxedEditing.Browser.VersionColumn.FillWidth", 1.f);
+	Set("SandboxedEditing.Browser.LastModifiedColumn.Padding", FMargin(8.f, 0.f, 0.f, 0.f));
+	Set("SandboxedEditing.Browser.LastModifiedColumn.FillWidth", 1.f);
+	
+}
+
+void FSandboxedEditingStyle::RegisterSharedStyle()
+{
+	
+}
 }
diff --git a/Engine/Plugins/Developer/Sandbox/SandboxedEditing/Source/SandboxedEditing/Private/SandboxedEditingStyle.h b/Engine/Plugins/Developer/Sandbox/SandboxedEditing/Source/SandboxedEditing/Private/SandboxedEditingStyle.h
--- a/Engine/Plugins/Developer/Sandbox/SandboxedEditing/Source/SandboxedEditing/Private/SandboxedEditingStyle.h
+++ b/Engine/Plugins/Developer/Sandbox/SandboxedEditing/Source/SandboxedEditing/Private/SandboxedEditingStyle.h
@@ -18,6 +18,11 @@ public:
 
 	FSandboxedEditingStyle();
 	virtual ~FSandboxedEditingStyle() override;
+	
+private:
+	
+	void RegisterBrowserStyle();
+	void RegisterSharedStyle();
 };
 }
 
diff --git a/Engine/Plugins/Developer/Sandbox/SandboxedEditing/Source/SandboxedEditing/Private/Utils/BreakBehavior.h b/Engine/Plugins/Developer/Sandbox/SandboxedEditing/Source/SandboxedEditing/Private/Utils/BreakBehavior.h
new file mode 100644
--- /dev/null
+++ b/Engine/Plugins/Developer/Sandbox/SandboxedEditing/Source/SandboxedEditing/Private/Utils/BreakBehavior.h
@@ -0,0 +1,14 @@
+﻿// Copyright Epic Games, Inc. All Rights Reserved.
+
+#pragma once
+
+#include "HAL/Platform.h"
+
+namespace UE::SandboxedEditing
+{
+enum class EBreakBehavior : uint8
+{
+	Break,
+	Continue
+};
+}
\ No newline at end of file
diff --git a/Engine/Plugins/Developer/Sandbox/SandboxedEditing/Source/SandboxedEditing/Private/Utils/WidgetUtils.cpp b/Engine/Plugins/Developer/Sandbox/SandboxedEditing/Source/SandboxedEditing/Private/Utils/WidgetUtils.cpp
new file mode 100644
--- /dev/null
+++ b/Engine/Plugins/Developer/Sandbox/SandboxedEditing/Source/SandboxedEditing/Private/Utils/WidgetUtils.cpp
@@ -0,0 +1,101 @@
+﻿// Copyright Epic Games, Inc. All Rights Reserved.
+
+#include "WidgetUtils.h"
+
+#include "SNegativeActionButton.h"
+#include "SPositiveActionButton.h"
+#include "Framework/SlateDelegates.h"
+#include "Framework/Commands/UICommandInfo.h"
+#include "Framework/Commands/UICommandList.h"
+#include "Internationalization/Text.h"
+#include "Misc/Attribute.h"
+#include "Styling/AppStyle.h"
+#include "Styling/SlateColor.h"
+#include "Widgets/Images/SImage.h"
+#include "Widgets/Input/SButton.h"
+
+namespace UE::SandboxedEditing
+{
+TSharedRef<SWidget> MakeIconButton(
+	const TAttribute<const FSlateBrush*>& Icon,
+	const TAttribute<FText>& Tooltip, 
+	const TAttribute<bool>& EnabledAttribute, 
+	const FOnClicked& OnClicked, 
+	const TAttribute<EVisibility>& Visibility
+	)
+{
+	return SNew(SButton)
+		.OnClicked(OnClicked)
+		.ButtonStyle(FAppStyle::Get(), TEXT("SimpleButton"))
+		.ToolTipText(Tooltip)
+		.Visibility(Visibility)
+		.IsEnabled(EnabledAttribute)
+		.VAlign(VAlign_Center)
+		.HAlign(HAlign_Center)
+		[
+			SNew(SImage)
+			.Image(Icon)
+			.ColorAndOpacity(FSlateColor::UseForeground())
+		];
+}
+
+TSharedRef<SWidget> MakeIconButtonByCommand(
+	const TAttribute<const FSlateBrush*>& Icon,
+	const TSharedPtr<FUICommandList>& InCommandList,
+	const TSharedPtr<FUICommandInfo>& InCommand,
+	const TAttribute<FText>& OverrideTooltip,
+	const TAttribute<EVisibility>& Visibility
+)
+{
+	return MakeIconButton(
+		Icon, 
+		OverrideTooltip.IsBound() || OverrideTooltip.IsSet() ? OverrideTooltip : InCommand->GetDescription(), 
+		TAttribute<bool>::CreateLambda([InCommandList, InCommand]
+		{
+			const FUIAction* Action = InCommandList->GetActionForCommand(InCommand);
+			return Action && Action->CanExecute();
+		}),
+		FOnClicked::CreateLambda([InCommandList, InCommand]
+		{
+			if (const FUIAction* Action = InCommandList->GetActionForCommand(InCommand))
+			{
+				Action->Execute();
+			}
+			return FReply::Handled();
+		}), 
+		Visibility)
+	;
+}
+
+TSharedRef<SWidget> MakePositiveActionButton(
+	const TAttribute<const FSlateBrush*>& Icon, 
+	const TAttribute<FText>& Tooltip, 
+	const TAttribute<bool>& EnabledAttribute, 
+	const FOnClicked& OnClicked, 
+	const TAttribute<EVisibility>& Visibility
+		)
+{
+	return SNew(SPositiveActionButton)
+		.OnClicked(OnClicked)
+		.ToolTipText(Tooltip)
+		.Visibility(Visibility)
+		.IsEnabled(EnabledAttribute)
+		.Icon(Icon);
+}
+
+TSharedRef<SWidget> MakeNegativeActionButton(
+	const TAttribute<const FSlateBrush*>& Icon, 
+	const TAttribute<FText>& Tooltip, 
+	const TAttribute<bool>& EnabledAttribute, 
+	const FOnClicked& OnClicked, 
+	const TAttribute<EVisibility>& Visibility
+		)
+{
+	return SNew(SNegativeActionButton)
+		.OnClicked(OnClicked)
+		.ToolTipText(Tooltip)
+		.Visibility(Visibility)
+		.IsEnabled(EnabledAttribute)
+		.Icon(Icon);
+}
+}
\ No newline at end of file
diff --git a/Engine/Plugins/Developer/Sandbox/SandboxedEditing/Source/SandboxedEditing/Private/Utils/WidgetUtils.h b/Engine/Plugins/Developer/Sandbox/SandboxedEditing/Source/SandboxedEditing/Private/Utils/WidgetUtils.h
new file mode 100644
--- /dev/null
+++ b/Engine/Plugins/Developer/Sandbox/SandboxedEditing/Source/SandboxedEditing/Private/Utils/WidgetUtils.h
@@ -0,0 +1,52 @@
+﻿// Copyright Epic Games, Inc. All Rights Reserved.
+
+#pragma once
+
+#include "Framework/SlateDelegates.h"
+#include "Templates/SharedPointerFwd.h"
+#include "Layout/Visibility.h"
+
+class FText;
+class FUICommandInfo;
+class FUICommandList;
+class SWidget;
+struct FSlateBrush;
+template<typename ObjectType> class TAttribute;
+
+namespace UE::SandboxedEditing
+{
+/** Makes an action button that only contains an icon. */
+TSharedRef<SWidget> MakeIconButton(
+	const TAttribute<const FSlateBrush*>& Icon,
+	const TAttribute<FText>& Tooltip, 
+	const TAttribute<bool>& EnabledAttribute, 
+	const FOnClicked& OnClicked, 
+	const TAttribute<EVisibility>& Visibility = EVisibility::Visible
+	);
+/** Overload that determines the IsEnabled and OnClicked parameters from the command list. */
+TSharedRef<SWidget> MakeIconButtonByCommand(
+	const TAttribute<const FSlateBrush*>& Icon,
+	const TSharedPtr<FUICommandList>& InCommandList,
+	const TSharedPtr<FUICommandInfo>& InCommand,
+	const TAttribute<FText>& OverrideTooltip = {},
+	const TAttribute<EVisibility>& Visibility = EVisibility::Visible
+);
+
+/** Makes a positive action button that only contains an icon */
+TSharedRef<SWidget> MakePositiveActionButton(
+	const TAttribute<const FSlateBrush*>& Icon, 
+	const TAttribute<FText>& Tooltip, 
+	const TAttribute<bool>& EnabledAttribute, 
+	const FOnClicked& OnClicked, 
+	const TAttribute<EVisibility>& Visibility = EVisibility::Visible
+	);
+
+/** Makes a negative action button that only contains an icon */
+TSharedRef<SWidget> MakeNegativeActionButton(
+	const TAttribute<const FSlateBrush*>& Icon, 
+	const TAttribute<FText>& Tooltip, 
+	const TAttribute<bool>& EnabledAttribute, 
+	const FOnClicked& OnClicked, 
+	const TAttribute<EVisibility>& Visibility = EVisibility::Visible
+	);
+}
diff --git a/Engine/Plugins/Developer/Sandbox/SandboxedEditing/Source/SandboxedEditing/SandboxedEditing.Build.cs b/Engine/Plugins/Developer/Sandbox/SandboxedEditing/Source/SandboxedEditing/SandboxedEditing.Build.cs
--- a/Engine/Plugins/Developer/Sandbox/SandboxedEditing/Source/SandboxedEditing/SandboxedEditing.Build.cs
+++ b/Engine/Plugins/Developer/Sandbox/SandboxedEditing/Source/SandboxedEditing/SandboxedEditing.Build.cs
@@ -22,7 +22,8 @@ namespace UnrealBuildTool.Rules
 					"Slate", 
 					"SlateCore", 
 					"ToolMenus", 
-					"ToolWidgets",
+					"ToolWidgets", 
+					"WorkspaceMenuStructure",
 				}
 			);
 		}
