Sandboxed Editing: UI that causes leaving of the current sandbox is now disabled when the sandbox is locked.

There will be a follow-up change that disables the UI when a sandbox that was not instigated by the Sandboxed Editing plugin is active.

#Synced-CL 47637010

diff --git a/Engine/Plugins/Developer/Sandbox/SandboxedEditing/Source/SandboxedEditing/Private/Features/Browser/Commands/BrowserCommandBindings.cpp b/Engine/Plugins/Developer/Sandbox/SandboxedEditing/Source/SandboxedEditing/Private/Features/Browser/Commands/BrowserCommandBindings.cpp
--- a/Engine/Plugins/Developer/Sandbox/SandboxedEditing/Source/SandboxedEditing/Private/Features/Browser/Commands/BrowserCommandBindings.cpp
+++ b/Engine/Plugins/Developer/Sandbox/SandboxedEditing/Source/SandboxedEditing/Private/Features/Browser/Commands/BrowserCommandBindings.cpp
@@ -52,6 +52,6 @@ void FBrowserCommandBindings::HandleLeaveSandbox()
 
 bool FBrowserCommandBindings::CanLeaveSandbox()
 {
-	return ActiveSandboxViewModel->HasActiveSandbox();
+	return ActiveSandboxViewModel->CanLeaveSandbox();
 }
 }
diff --git a/Engine/Plugins/Developer/Sandbox/SandboxedEditing/Source/SandboxedEditing/Private/Features/Browser/ViewModels/Active/ActiveSandboxDetailsViewModel.cpp b/Engine/Plugins/Developer/Sandbox/SandboxedEditing/Source/SandboxedEditing/Private/Features/Browser/ViewModels/Active/ActiveSandboxDetailsViewModel.cpp
--- a/Engine/Plugins/Developer/Sandbox/SandboxedEditing/Source/SandboxedEditing/Private/Features/Browser/ViewModels/Active/ActiveSandboxDetailsViewModel.cpp
+++ b/Engine/Plugins/Developer/Sandbox/SandboxedEditing/Source/SandboxedEditing/Private/Features/Browser/ViewModels/Active/ActiveSandboxDetailsViewModel.cpp
@@ -20,6 +20,11 @@ void FActiveSandboxDetailsViewModel::LeaveSandbox()
 	Model->LeaveSandbox();
 }
 
+bool FActiveSandboxDetailsViewModel::CanLeaveSandbox(FText* OutReason) const
+{
+	return Model->CanLeaveSandbox(OutReason);
+}
+
 FString FActiveSandboxDetailsViewModel::GetSandboxName() const
 {
 	return Model->GetActiveSandboxName();
diff --git a/Engine/Plugins/Developer/Sandbox/SandboxedEditing/Source/SandboxedEditing/Private/Features/Browser/ViewModels/Active/ActiveSandboxDetailsViewModel.h b/Engine/Plugins/Developer/Sandbox/SandboxedEditing/Source/SandboxedEditing/Private/Features/Browser/ViewModels/Active/ActiveSandboxDetailsViewModel.h
--- a/Engine/Plugins/Developer/Sandbox/SandboxedEditing/Source/SandboxedEditing/Private/Features/Browser/ViewModels/Active/ActiveSandboxDetailsViewModel.h
+++ b/Engine/Plugins/Developer/Sandbox/SandboxedEditing/Source/SandboxedEditing/Private/Features/Browser/ViewModels/Active/ActiveSandboxDetailsViewModel.h
@@ -22,6 +22,8 @@ public:
 	
 	/** Leaves the active sandbox. */
 	void LeaveSandbox();
+	/** @return Whether the sandbox can be left */
+	bool CanLeaveSandbox(FText* OutReason = nullptr) const; 
 	
 	/** @return Name of the active sandbox. */
 	FString GetSandboxName() const;
diff --git a/Engine/Plugins/Developer/Sandbox/SandboxedEditing/Source/SandboxedEditing/Private/Features/Browser/ViewModels/SandboxCreationWorkflow.cpp b/Engine/Plugins/Developer/Sandbox/SandboxedEditing/Source/SandboxedEditing/Private/Features/Browser/ViewModels/SandboxCreationWorkflow.cpp
--- a/Engine/Plugins/Developer/Sandbox/SandboxedEditing/Source/SandboxedEditing/Private/Features/Browser/ViewModels/SandboxCreationWorkflow.cpp
+++ b/Engine/Plugins/Developer/Sandbox/SandboxedEditing/Source/SandboxedEditing/Private/Features/Browser/ViewModels/SandboxCreationWorkflow.cpp
@@ -33,7 +33,7 @@ void FSandboxCreationWorkflow::Cancel()
 
 bool FSandboxCreationWorkflow::CanConfirm(FText* OutReason) const
 {
-	return IsValidName(Name, OutReason);
+	return IsValidName(Name, OutReason) && Model->CanCreateNewSandbox(Name, OutReason);
 }
 
 void FSandboxCreationWorkflow::SetName(const FString& InNewName)
diff --git a/Engine/Plugins/Developer/Sandbox/SandboxedEditing/Source/SandboxedEditing/Private/Features/Browser/Widgets/Sandboxed/SActiveSandboxToolbar.cpp b/Engine/Plugins/Developer/Sandbox/SandboxedEditing/Source/SandboxedEditing/Private/Features/Browser/Widgets/Sandboxed/SActiveSandboxToolbar.cpp
--- a/Engine/Plugins/Developer/Sandbox/SandboxedEditing/Source/SandboxedEditing/Private/Features/Browser/Widgets/Sandboxed/SActiveSandboxToolbar.cpp
+++ b/Engine/Plugins/Developer/Sandbox/SandboxedEditing/Source/SandboxedEditing/Private/Features/Browser/Widgets/Sandboxed/SActiveSandboxToolbar.cpp
@@ -54,7 +54,8 @@ void SActiveSandboxToolbar::Construct(const FArguments& InArgs, const TSharedRef
 		[
 			MakeIconButtonByCommand(
 				FSandboxedEditingStyle::Get().GetBrush("SandboxedEditing.Browser.LeaveSandbox"),
-				InArgs._CommandList, FBrowserCommands::Get().LeaveSandbox
+				InArgs._CommandList, FBrowserCommands::Get().LeaveSandbox, 
+				TAttribute<FText>::CreateSP(this, &SActiveSandboxToolbar::GetLeaveButtonToolTipText)
 				)
 		]
 	];
@@ -72,4 +73,12 @@ TSharedRef<SWidget> SActiveSandboxToolbar::MakeNameWidget() const
 			.Text_Lambda([this]{ return FText::AsCultureInvariant(ActiveSandboxViewModel->GetSandboxName()); })
 		];
 }
+
+FText SActiveSandboxToolbar::GetLeaveButtonToolTipText() const
+{
+	FText Reason;
+	return ActiveSandboxViewModel->CanLeaveSandbox(&Reason) 
+		? FBrowserCommands::Get().LeaveSandbox->GetDescription() 
+		: Reason;
+}
 }
diff --git a/Engine/Plugins/Developer/Sandbox/SandboxedEditing/Source/SandboxedEditing/Private/Features/Browser/Widgets/Sandboxed/SActiveSandboxToolbar.h b/Engine/Plugins/Developer/Sandbox/SandboxedEditing/Source/SandboxedEditing/Private/Features/Browser/Widgets/Sandboxed/SActiveSandboxToolbar.h
--- a/Engine/Plugins/Developer/Sandbox/SandboxedEditing/Source/SandboxedEditing/Private/Features/Browser/Widgets/Sandboxed/SActiveSandboxToolbar.h
+++ b/Engine/Plugins/Developer/Sandbox/SandboxedEditing/Source/SandboxedEditing/Private/Features/Browser/Widgets/Sandboxed/SActiveSandboxToolbar.h
@@ -30,6 +30,9 @@ private:
 	
 	/** @return Widget displaying the active sandbox name. */
 	TSharedRef<SWidget> MakeNameWidget() const;
+	
+	/** @return Tooltip text to display for the leave button */
+	FText GetLeaveButtonToolTipText() const;
 };
 }
 
diff --git a/Engine/Plugins/Developer/Sandbox/SandboxedEditing/Source/SandboxedEditing/Private/Features/StatusBar/Widgets/SSandboxStatusBarWidget.cpp b/Engine/Plugins/Developer/Sandbox/SandboxedEditing/Source/SandboxedEditing/Private/Features/StatusBar/Widgets/SSandboxStatusBarWidget.cpp
--- a/Engine/Plugins/Developer/Sandbox/SandboxedEditing/Source/SandboxedEditing/Private/Features/StatusBar/Widgets/SSandboxStatusBarWidget.cpp
+++ b/Engine/Plugins/Developer/Sandbox/SandboxedEditing/Source/SandboxedEditing/Private/Features/StatusBar/Widgets/SSandboxStatusBarWidget.cpp
@@ -69,8 +69,14 @@ TSharedRef<SWidget> SSandboxStatusBarWidget::CreateStatusBarMenu()
 	{
 		FToolMenuSection& Section = Menu->AddSection("SandboxedEditingSection", LOCTEXT("SandboxedEditing.Sandbox", "Sandbox"));
 
+		const TAttribute<FText> LeaveToolTipOverride = TAttribute<FText>::CreateLambda([Model = ViewModel]
+		{
+			FText Reason;
+			return Model->CanLeaveSandbox(&Reason) ? FStatusBarCommands::Get().LeaveSandbox->GetDescription() : Reason;
+		});
+		
 		Section.AddMenuEntry(FStatusBarCommands::Get().OpenCreateNewSandboxDialog);
-		Section.AddMenuEntry(FStatusBarCommands::Get().LeaveSandbox);
+		Section.AddMenuEntry(FStatusBarCommands::Get().LeaveSandbox, {}, LeaveToolTipOverride);
 		Section.AddMenuEntry(FStatusBarCommands::Get().PersistAll);
 		Section.AddMenuEntry(FStatusBarCommands::Get().DiscardAll);
 
diff --git a/Engine/Plugins/Developer/Sandbox/SandboxedEditing/Source/SandboxedEditing/Private/Framework/Models/SandboxSystemModel.cpp b/Engine/Plugins/Developer/Sandbox/SandboxedEditing/Source/SandboxedEditing/Private/Framework/Models/SandboxSystemModel.cpp
--- a/Engine/Plugins/Developer/Sandbox/SandboxedEditing/Source/SandboxedEditing/Private/Framework/Models/SandboxSystemModel.cpp
+++ b/Engine/Plugins/Developer/Sandbox/SandboxedEditing/Source/SandboxedEditing/Private/Framework/Models/SandboxSystemModel.cpp
@@ -10,6 +10,7 @@
 #include "ISandboxInstance.h"
 #include "ISandboxManager.h"
 #include "ISandboxRepository.h"
+#include "Interface/ISandboxLock.h"
 #include "Misc/Paths.h"
 #include "Templates/Function.h"
 #include "Types/EBreakBehavior.h"
@@ -66,9 +67,9 @@ bool FSandboxSystemModel::CreateNewSandbox(const FCreateSandboxArgs& InArgs)
 	IFileSandboxCoreModule& Module = IFileSandboxCoreModule::Get();
 	ISandboxManager& SandboxManager = Module.GetSandboxManager();
 
-	if (SandboxManager.HasActiveSandbox())
+	if (!SandboxManager.LeaveSandbox())
 	{
-		SandboxManager.LeaveSandbox();
+		return false;
 	}
 
 	const FNewSandboxArgs Args(InArgs.Name, InArgs.Description);
@@ -85,6 +86,11 @@ bool FSandboxSystemModel::CanCreateNewSandbox(const FString& InName, FText* Reas
 		*Reason = FileSandboxCore::FormatDirectorySuitabilityAsReason(Suitability, InName);
 	}
 	
+	if (!CanLeaveSandbox(Reason))
+	{
+		return false;
+	}
+	
 	return bIsSuitable;
 }
 
@@ -122,13 +128,35 @@ void FSandboxSystemModel::LeaveSandbox() const
 	SandboxManager.LeaveSandbox();
 }
 
-bool FSandboxSystemModel::CanLeaveSandbox() const
+#define SET_REASON(ReasonText) if (OutReason){ *OutReason = ReasonText; }
+bool FSandboxSystemModel::CanLeaveSandbox(FText* OutReason) const
 {
 	using namespace FileSandboxCore;
 	IFileSandboxCoreModule& Module = IFileSandboxCoreModule::Get();
 	ISandboxManager& SandboxManager = Module.GetSandboxManager();
-	return SandboxManager.GetActiveSandboxInstance() != nullptr;
+	
+	if (SandboxManager.GetActiveSandboxInstance() == nullptr)
+	{
+		SET_REASON(LOCTEXT("LeaveReason.NoSandbox", "No sandbox is active."));
+		return false;
+	}
+	
+	if (ISandboxLock* Lock = SandboxManager.GetActiveLock()
+		; Lock && !Lock->CanLeaveSandbox(OutReason))
+	{
+		return false;
+	}
+	
+	// If this happens, we should investigate why the sandbox cannot be left even though the lock said it's fine. Return a proper a reason.
+	if (ensure(!SandboxManager.CanLeaveSandbox()))
+	{
+		SET_REASON(LOCTEXT("LeaveReason.OtherReason", "Cannot leave sandbox: no further info available."));
+		return false;
+	}
+	
+	return true;
 }
+#undef SET_REASON
 
 void FSandboxSystemModel::DeleteSandbox(const FString& InSandboxRoot) const
 {
diff --git a/Engine/Plugins/Developer/Sandbox/SandboxedEditing/Source/SandboxedEditing/Private/Framework/Models/SandboxSystemModel.h b/Engine/Plugins/Developer/Sandbox/SandboxedEditing/Source/SandboxedEditing/Private/Framework/Models/SandboxSystemModel.h
--- a/Engine/Plugins/Developer/Sandbox/SandboxedEditing/Source/SandboxedEditing/Private/Framework/Models/SandboxSystemModel.h
+++ b/Engine/Plugins/Developer/Sandbox/SandboxedEditing/Source/SandboxedEditing/Private/Framework/Models/SandboxSystemModel.h
@@ -45,7 +45,9 @@ public:
 	/** Leaves the current sandbox. */
 	void LeaveSandbox() const;
 	/** @return Whether the sandbox can be left (false if there is no active one). */
-	bool CanLeaveSandbox() const;
+	bool CanLeaveSandbox() const { return CanLeaveSandbox(nullptr); }
+	/** @return Whether the sandbox can be left (false if there is no active one). */
+	bool CanLeaveSandbox(FText* OutReason) const;
 	
 	/** Deletes the sandbox. */
 	void DeleteSandbox(const FString& InSandboxRoot) const;
