File Sandbox: Fix potential crash due to bad API by replacing multicast with unicast delegate. The API allowed binding to the multicast but invoking it causes the mulitcast to be destroyed... which would cause a crash if more than 1 subscribers are present.

#Synced-CL 47637010

diff --git a/Engine/Plugins/Developer/Sandbox/SandboxedEditing/Source/SandboxedEditing/Private/Features/Browser/ViewModels/SandboxControlsViewModel.cpp b/Engine/Plugins/Developer/Sandbox/SandboxedEditing/Source/SandboxedEditing/Private/Features/Browser/ViewModels/SandboxControlsViewModel.cpp
--- a/Engine/Plugins/Developer/Sandbox/SandboxedEditing/Source/SandboxedEditing/Private/Features/Browser/ViewModels/SandboxControlsViewModel.cpp
+++ b/Engine/Plugins/Developer/Sandbox/SandboxedEditing/Source/SandboxedEditing/Private/Features/Browser/ViewModels/SandboxControlsViewModel.cpp
@@ -12,8 +12,10 @@ void FSandboxControlsViewModel::StartCreationWorkflow()
 {
 	if (CanStartCreationWorkflow())
 	{
-		CurrentCreationWorkflow = MakeUnique<FSandboxCreationWorkflow>(Model);
-		CurrentCreationWorkflow->OnWorkflowEnded().AddRaw(this, &FSandboxControlsViewModel::OnWorkflowEnded);
+		CurrentCreationWorkflow = MakeUnique<FSandboxCreationWorkflow>(
+			Model,
+			FSimpleDelegate::CreateRaw(this, &FSandboxControlsViewModel::OnWorkflowEnded)
+			);
 		OnCreationWorkflowStartedDelegate.Broadcast(*CurrentCreationWorkflow);
 	}
 }
diff --git a/Engine/Plugins/Developer/Sandbox/SandboxedEditing/Source/SandboxedEditing/Private/Features/Browser/ViewModels/SandboxCreationWorkflow.cpp b/Engine/Plugins/Developer/Sandbox/SandboxedEditing/Source/SandboxedEditing/Private/Features/Browser/ViewModels/SandboxCreationWorkflow.cpp
--- a/Engine/Plugins/Developer/Sandbox/SandboxedEditing/Source/SandboxedEditing/Private/Features/Browser/ViewModels/SandboxCreationWorkflow.cpp
+++ b/Engine/Plugins/Developer/Sandbox/SandboxedEditing/Source/SandboxedEditing/Private/Features/Browser/ViewModels/SandboxCreationWorkflow.cpp
@@ -7,8 +7,9 @@
 
 namespace UE::SandboxedEditing
 {
-FSandboxCreationWorkflow::FSandboxCreationWorkflow(const TSharedRef<FSandboxSystemModel>& InModel)
+FSandboxCreationWorkflow::FSandboxCreationWorkflow(const TSharedRef<FSandboxSystemModel>& InModel, FSimpleDelegate InOnWorkflowEndedDelegate)
 	: Model(InModel)
+	, OnWorkflowEndedDelegate(MoveTemp(InOnWorkflowEndedDelegate))
 {}
 
 void FSandboxCreationWorkflow::Confirm()
@@ -23,12 +24,12 @@ void FSandboxCreationWorkflow::Confirm()
 		return;
 	}
 	
-	OnWorkflowEndedDelegate.Broadcast();
+	OnWorkflowEndedDelegate.ExecuteIfBound();
 }
 
 void FSandboxCreationWorkflow::Cancel()
 {
-	OnWorkflowEndedDelegate.Broadcast();
+	OnWorkflowEndedDelegate.ExecuteIfBound();
 }
 
 bool FSandboxCreationWorkflow::CanConfirm(FText* OutReason) const
diff --git a/Engine/Plugins/Developer/Sandbox/SandboxedEditing/Source/SandboxedEditing/Private/Features/Browser/ViewModels/SandboxCreationWorkflow.h b/Engine/Plugins/Developer/Sandbox/SandboxedEditing/Source/SandboxedEditing/Private/Features/Browser/ViewModels/SandboxCreationWorkflow.h
--- a/Engine/Plugins/Developer/Sandbox/SandboxedEditing/Source/SandboxedEditing/Private/Features/Browser/ViewModels/SandboxCreationWorkflow.h
+++ b/Engine/Plugins/Developer/Sandbox/SandboxedEditing/Source/SandboxedEditing/Private/Features/Browser/ViewModels/SandboxCreationWorkflow.h
@@ -19,7 +19,7 @@ class FSandboxCreationWorkflow : public FNoncopyable
 {
 public:
 	
-	explicit FSandboxCreationWorkflow(const TSharedRef<FSandboxSystemModel>& InModel);
+	explicit FSandboxCreationWorkflow(const TSharedRef<FSandboxSystemModel>& InModel, FSimpleDelegate InOnWorkflowEndedDelegate);
 	
 	/** Finishes the workflow, if possible, resulting in the creation of the sandbox. */
 	void Confirm();
@@ -37,16 +37,13 @@ public:
 	/** Sets the description. There are no restrictions on the description. */
 	void SetDescription(const FString& InNewDescription);
 	
-	/** Invoked when the workflow ends. */
-	FSimpleMulticastDelegate& OnWorkflowEnded() { return OnWorkflowEndedDelegate; }
-	
 private:  
 	
 	/** The model used to interact with the sandbox system. */
 	const TSharedRef<FSandboxSystemModel> Model;
 	
 	/** Invoked when the workflow ends. */
-	FSimpleMulticastDelegate OnWorkflowEndedDelegate;
+	const FSimpleDelegate OnWorkflowEndedDelegate;
 	
 	/** The name the user wants to give. */
 	FString Name;
