Sandboxed Editing: Show overlay when external sandbox is active

#Synced-CL 47637010

diff --git a/Engine/Plugins/Developer/Sandbox/SandboxedEditing/Source/SandboxedEditing/Private/Features/Browser/BrowserFeature.cpp b/Engine/Plugins/Developer/Sandbox/SandboxedEditing/Source/SandboxedEditing/Private/Features/Browser/BrowserFeature.cpp
--- a/Engine/Plugins/Developer/Sandbox/SandboxedEditing/Source/SandboxedEditing/Private/Features/Browser/BrowserFeature.cpp
+++ b/Engine/Plugins/Developer/Sandbox/SandboxedEditing/Source/SandboxedEditing/Private/Features/Browser/BrowserFeature.cpp
@@ -41,10 +41,13 @@ static TMap<FName, TSharedRef<ISandboxWidgetFactory>> TransformToWidgetFactories
 }
 }
 	
-FBrowserFeature::FBrowserFeature(const TSharedRef<FSandboxSystemModel>& InSandboxModel)
+FBrowserFeature::FBrowserFeature(
+	const TSharedRef<FSandboxSystemModel>& InSandboxModel, 
+	const TSharedRef<FileSandboxUI::IExternalSandboxActiveViewModel>& InExternalSandboxViewModel
+	)
 	: SandboxModel(InSandboxModel)
 	, BrowserColumns(MakeBuiltInColumns())
-	, ViewModels(BrowserDetail::TransformToBehaviours(BrowserColumns), SandboxModel)
+	, ViewModels(BrowserDetail::TransformToBehaviours(BrowserColumns), SandboxModel, InExternalSandboxViewModel)
 	, CommandBindings(MakeShared<FBrowserCommandBindings>(ViewModels))
 {
 	RegisterTabSpawner();
@@ -55,6 +58,11 @@ FBrowserFeature::~FBrowserFeature()
 	UnregisterTabSpawner();
 }
 
+void FBrowserFeature::SummonUI() const
+{
+	FGlobalTabmanager::Get()->TryInvokeTab(BrowserDetail::SandboxedEditingTabName.Resolve());
+}
+
 void FBrowserFeature::RegisterTabSpawner()
 {
 	FTabSpawnerEntry& BrowserSpawnerEntry = FGlobalTabmanager::Get()->RegisterNomadTabSpawner(
diff --git a/Engine/Plugins/Developer/Sandbox/SandboxedEditing/Source/SandboxedEditing/Private/Features/Browser/BrowserFeature.h b/Engine/Plugins/Developer/Sandbox/SandboxedEditing/Source/SandboxedEditing/Private/Features/Browser/BrowserFeature.h
--- a/Engine/Plugins/Developer/Sandbox/SandboxedEditing/Source/SandboxedEditing/Private/Features/Browser/BrowserFeature.h
+++ b/Engine/Plugins/Developer/Sandbox/SandboxedEditing/Source/SandboxedEditing/Private/Features/Browser/BrowserFeature.h
@@ -23,10 +23,20 @@ class FSandboxSystemModel;
 class FBrowserFeature : public FNoncopyable
 {
 public:
-
-	explicit FBrowserFeature(const TSharedRef<FSandboxSystemModel>& InSandboxModel);
+	
+	/**
+	 * @param InSandboxModel Model that interacts with the core sandbox system.
+	 * @param InExternalSandboxViewModel View model to let us know when an external sandbox is active.
+	 */
+	explicit FBrowserFeature(
+		const TSharedRef<FSandboxSystemModel>& InSandboxModel, 
+		const TSharedRef<FileSandboxUI::IExternalSandboxActiveViewModel>& InExternalSandboxViewModel
+		);
 	~FBrowserFeature();
 	
+	/** Summons the browser tab. */
+	void SummonUI() const;
+	
 private:
 	
 	/** The shared sandbox model. */
diff --git a/Engine/Plugins/Developer/Sandbox/SandboxedEditing/Source/SandboxedEditing/Private/Features/Browser/ViewModels/BrowserViewModels.h b/Engine/Plugins/Developer/Sandbox/SandboxedEditing/Source/SandboxedEditing/Private/Features/Browser/ViewModels/BrowserViewModels.h
--- a/Engine/Plugins/Developer/Sandbox/SandboxedEditing/Source/SandboxedEditing/Private/Features/Browser/ViewModels/BrowserViewModels.h
+++ b/Engine/Plugins/Developer/Sandbox/SandboxedEditing/Source/SandboxedEditing/Private/Features/Browser/ViewModels/BrowserViewModels.h
@@ -10,6 +10,11 @@
 #include "SandboxMetaDataViewModel.h"
 #include "Templates/SharedPointer.h"
 
+namespace UE::FileSandboxUI
+{
+class IExternalSandboxActiveViewModel;
+}
+
 namespace UE::SandboxedEditing
 {
 /** Holds all view models used by the root UI. Make sharing view models easier. */
@@ -32,8 +37,14 @@ struct FBrowserViewModels
 	const TSharedRef<FActiveSandboxTrackerViewModel> ActiveSandboxTrackerViewModel;
 	/** Used for widgets that display info about the currently active sandbox. */
 	const TSharedRef<FActiveSandboxDetailsViewModel> ActiveSandboxDetailsViewModel;
+	/** Knows about when an extern sandbox is being displayed. */
+	const TSharedRef<FileSandboxUI::IExternalSandboxActiveViewModel> ExternalSandboxViewModel;
 	
-	explicit FBrowserViewModels(TMap<FName, TSharedRef<ISandboxColumnBehavior>> InColumns, const TSharedRef<FSandboxSystemModel>& InModel)
+	explicit FBrowserViewModels(
+		TMap<FName, TSharedRef<ISandboxColumnBehavior>> InColumns, 
+		const TSharedRef<FSandboxSystemModel>& InModel,
+		const TSharedRef<FileSandboxUI::IExternalSandboxActiveViewModel>& InExternalSandboxViewModel
+		)
 		: FilterSandboxViewModel(MakeShared<FFilterSandboxViewModel>([&InColumns]
 			{
 				TArray<TSharedRef<ISandboxColumnBehavior>> Columns;
@@ -45,6 +56,7 @@ struct FBrowserViewModels
 		, MetaDataViewModel(MakeShared<FSandboxMetaDataViewModel>(InModel))
 		, ActiveSandboxTrackerViewModel(MakeShared<FActiveSandboxTrackerViewModel>(InModel))
 		, ActiveSandboxDetailsViewModel(MakeShared<FActiveSandboxDetailsViewModel>(InModel))
+		, ExternalSandboxViewModel(InExternalSandboxViewModel)
 	{}
 };
 }
diff --git a/Engine/Plugins/Developer/Sandbox/SandboxedEditing/Source/SandboxedEditing/Private/Features/Browser/Widgets/SSandboxedEditingRoot.cpp b/Engine/Plugins/Developer/Sandbox/SandboxedEditing/Source/SandboxedEditing/Private/Features/Browser/Widgets/SSandboxedEditingRoot.cpp
--- a/Engine/Plugins/Developer/Sandbox/SandboxedEditing/Source/SandboxedEditing/Private/Features/Browser/Widgets/SSandboxedEditingRoot.cpp
+++ b/Engine/Plugins/Developer/Sandbox/SandboxedEditing/Source/SandboxedEditing/Private/Features/Browser/Widgets/SSandboxedEditingRoot.cpp
@@ -2,12 +2,15 @@
 
 #include "SSandboxedEditingRoot.h"
 
+#include "EntryPoint/EntryPointWidgetFactory.h"
+#include "EntryPoint/IExternalSandboxActiveViewModel.h"
 #include "Features/Browser/Commands/BrowserCommandBindings.h"
-#include "Features/Browser/ViewModels/BrowserViewModels.h"
 #include "Features/Browser/ViewModels/Active/ActiveSandboxTrackerViewModel.h"
+#include "Features/Browser/ViewModels/BrowserViewModels.h"
 #include "Features/Browser/Widgets/Unsandboxed/SSandboxedEditingBrowser.h"
 #include "Sandboxed/SActiveSandbox.h"
-#include "Widgets/Layout/SWidgetSwitcher.h"
+#include "Widgets/Layout/SBox.h"
+#include "Widgets/SOverlay.h"
 
 namespace UE::SandboxedEditing
 {
@@ -15,10 +18,12 @@ void SSandboxedEditingRoot::Construct(const FArguments& InArgs, const FBrowserVi
 {
 	CommandBindings = InArgs._CommandList;
 	ActiveSandboxTrackerViewModel = InViewModels.ActiveSandboxTrackerViewModel;
+	ExternalSandboxViewModel = InViewModels.ExternalSandboxViewModel;
 	ActiveSandboxViewModels.Emplace(InViewModels.ActiveSandboxDetailsViewModel, InViewModels.MetaDataViewModel);
 	
-	ActiveSandboxTrackerViewModel->OnLoadSandbox().AddSP(this, &SSandboxedEditingRoot::ShowActiveSandbox);
+	ActiveSandboxTrackerViewModel->OnLoadSandbox().AddSP(this, &SSandboxedEditingRoot::UpdateWidgetContent);
 	ActiveSandboxTrackerViewModel->OnLeaveSandbox().AddSP(this, &SSandboxedEditingRoot::ShowBrowser);
+	ExternalSandboxViewModel->OnVisibilityChanged().AddSP(this, &SSandboxedEditingRoot::UpdateWidgetContent);
 	
 	const SSandboxedEditingBrowser::FViewModels ViewModels 
 	{ 
@@ -66,7 +71,21 @@ void SSandboxedEditingRoot::ShowBrowser()
 {
 	ChildSlot
 	[
-		Browser.ToSharedRef()
+		SNew(SOverlay)
+
+		+SOverlay::Slot()
+		[
+			Browser.ToSharedRef()
+		]
+		
+		+SOverlay::Slot()
+		[
+			SNew(SBox)
+			.Visibility_Lambda([this]{ return ExternalSandboxViewModel->IsExternalSandboxActive() ? EVisibility::Visible : EVisibility::Collapsed; })
+			[
+				FileSandboxUI::MakeExternalSandboxActiveOverlay(ExternalSandboxViewModel.ToSharedRef())
+			]
+		]
 	];
 }
 }
diff --git a/Engine/Plugins/Developer/Sandbox/SandboxedEditing/Source/SandboxedEditing/Private/Features/Browser/Widgets/SSandboxedEditingRoot.h b/Engine/Plugins/Developer/Sandbox/SandboxedEditing/Source/SandboxedEditing/Private/Features/Browser/Widgets/SSandboxedEditingRoot.h
--- a/Engine/Plugins/Developer/Sandbox/SandboxedEditing/Source/SandboxedEditing/Private/Features/Browser/Widgets/SSandboxedEditingRoot.h
+++ b/Engine/Plugins/Developer/Sandbox/SandboxedEditing/Source/SandboxedEditing/Private/Features/Browser/Widgets/SSandboxedEditingRoot.h
@@ -9,6 +9,7 @@
 #include "Widgets/SCompoundWidget.h"
 
 class FUICommandList;
+namespace UE::FileSandboxUI { class IExternalSandboxActiveViewModel; }
 
 namespace UE::SandboxedEditing
 {
@@ -45,6 +46,8 @@ private:
 	
 	/** Knows when the active sandbox changes. */
 	TSharedPtr<FActiveSandboxTrackerViewModel> ActiveSandboxTrackerViewModel;
+	/** Knows when an external sandbox is active. */
+	TSharedPtr<FileSandboxUI::IExternalSandboxActiveViewModel> ExternalSandboxViewModel;
 	/** The view models required to construct the active sandbox widget. Set once Construct has been called. */
 	TOptional<SActiveSandbox::FViewModels> ActiveSandboxViewModels;
 	
diff --git a/Engine/Plugins/Developer/Sandbox/SandboxedEditing/Source/SandboxedEditing/Private/Framework/EntryPoint/SandboxedEditingEntryPoint.cpp b/Engine/Plugins/Developer/Sandbox/SandboxedEditing/Source/SandboxedEditing/Private/Framework/EntryPoint/SandboxedEditingEntryPoint.cpp
new file mode 100644
--- /dev/null
+++ b/Engine/Plugins/Developer/Sandbox/SandboxedEditing/Source/SandboxedEditing/Private/Framework/EntryPoint/SandboxedEditingEntryPoint.cpp
@@ -0,0 +1,30 @@
+﻿// Copyright Epic Games, Inc. All Rights Reserved.
+
+#include "SandboxedEditingEntryPoint.h"
+
+#include "Data/SandboxMetaData.h"
+#include "Framework/Models/SandboxTagging.h"
+#include "ISandboxInstance.h"
+#include "Internationalization/Internationalization.h"
+
+#define LOCTEXT_NAMESPACE "FSandboxedEditingEntryPoint"
+
+namespace UE::SandboxedEditing
+{
+void FSandboxedEditingEntryPoint::SummonProviderUI()
+{
+	OnRequestSummonUIDelegate.Broadcast();
+}
+
+FText FSandboxedEditingEntryPoint::GetEntryPointLabel() const
+{
+	return LOCTEXT("EntryPoint", "Sandboxed Editing");
+}
+
+bool FSandboxedEditingEntryPoint::OwnsSandbox(const FileSandboxCore::ISandboxInstance& InSandbox) const
+{
+	return IsOwnedBySandboxedEditing(InSandbox.GetInitialMetaData());
+}
+}
+
+#undef LOCTEXT_NAMESPACE
\ No newline at end of file
diff --git a/Engine/Plugins/Developer/Sandbox/SandboxedEditing/Source/SandboxedEditing/Private/Framework/EntryPoint/SandboxedEditingEntryPoint.h b/Engine/Plugins/Developer/Sandbox/SandboxedEditing/Source/SandboxedEditing/Private/Framework/EntryPoint/SandboxedEditingEntryPoint.h
new file mode 100644
--- /dev/null
+++ b/Engine/Plugins/Developer/Sandbox/SandboxedEditing/Source/SandboxedEditing/Private/Framework/EntryPoint/SandboxedEditingEntryPoint.h
@@ -0,0 +1,32 @@
+﻿// Copyright Epic Games, Inc. All Rights Reserved.
+
+#pragma once
+
+#include "Delegates/Delegate.h"
+#include "EntryPoint/ISandboxEntryPoint.h"
+
+namespace UE::SandboxedEditing
+{
+class FBrowserFeature;
+
+/** Represents entry point for sandboxed in Sandboxed Editing. */
+class FSandboxedEditingEntryPoint : public FileSandboxUI::ISandboxEntryPoint
+{
+public:
+	
+	//~ Begin ISandboxEntryPoint Interface
+	virtual void SummonProviderUI() override;
+	virtual FText GetEntryPointLabel() const override;
+	virtual bool OwnsSandbox(const FileSandboxCore::ISandboxInstance& InSandbox) const override;
+	//~ End ISandboxEntryPoint Interface
+	
+	/** Event invoked when the UI is supposed to be summoned. */
+	FSimpleMulticastDelegate& OnRequestSummonUI() { return OnRequestSummonUIDelegate; }
+	
+private:
+	
+	/** Event invoked when the UI is supposed to be summoned. */
+	FSimpleMulticastDelegate OnRequestSummonUIDelegate;
+};
+}
+
diff --git a/Engine/Plugins/Developer/Sandbox/SandboxedEditing/Source/SandboxedEditing/Private/Framework/EntryPoint/SandboxedEditingSandboxIntegration.cpp b/Engine/Plugins/Developer/Sandbox/SandboxedEditing/Source/SandboxedEditing/Private/Framework/EntryPoint/SandboxedEditingSandboxIntegration.cpp
new file mode 100644
--- /dev/null
+++ b/Engine/Plugins/Developer/Sandbox/SandboxedEditing/Source/SandboxedEditing/Private/Framework/EntryPoint/SandboxedEditingSandboxIntegration.cpp
@@ -0,0 +1,30 @@
+﻿// Copyright Epic Games, Inc. All Rights Reserved.
+
+#include "SandboxedEditingSandboxIntegration.h"
+
+#include "EntryPoint/EntryPointWidgetFactory.h"
+#include "EntryPoint/ISandboxEntryPointRegistry.h"
+#include "IFileSandboxUIModule.h"
+#include "SandboxedEditingEntryPoint.h"
+
+namespace UE::SandboxedEditing
+{
+FSandboxedEditingSandboxIntegration::FSandboxedEditingSandboxIntegration()
+	: EntryPoint(MakeShared<FSandboxedEditingEntryPoint>())
+	, ExternalSandboxViewModel(FileSandboxUI::MakeExternalSandboxActiveViewModel(EntryPoint))
+{
+	using namespace UE::FileSandboxUI;
+	ISandboxEntryPointRegistry& Registry = IFileSandboxUIModule::Get().GetEntryPointRegistry();
+	Registry.RegisterEntryPoint(EntryPoint);
+}
+
+FSandboxedEditingSandboxIntegration::~FSandboxedEditingSandboxIntegration()
+{
+	using namespace UE::FileSandboxUI;
+	if (IFileSandboxUIModule::IsAvailable())
+	{
+		ISandboxEntryPointRegistry& Registry = IFileSandboxUIModule::Get().GetEntryPointRegistry();
+		Registry.UnregisterEntryPoint(EntryPoint);
+	}
+}
+}
diff --git a/Engine/Plugins/Developer/Sandbox/SandboxedEditing/Source/SandboxedEditing/Private/Framework/EntryPoint/SandboxedEditingSandboxIntegration.h b/Engine/Plugins/Developer/Sandbox/SandboxedEditing/Source/SandboxedEditing/Private/Framework/EntryPoint/SandboxedEditingSandboxIntegration.h
new file mode 100644
--- /dev/null
+++ b/Engine/Plugins/Developer/Sandbox/SandboxedEditing/Source/SandboxedEditing/Private/Framework/EntryPoint/SandboxedEditingSandboxIntegration.h
@@ -0,0 +1,34 @@
+﻿// Copyright Epic Games, Inc. All Rights Reserved.
+
+#pragma once
+
+#include "Templates/SharedPointer.h"
+#include "Templates/UnrealTemplate.h"
+
+namespace UE::FileSandboxUI { class IExternalSandboxActiveViewModel; }
+
+namespace UE::SandboxedEditing
+{
+class FSandboxedEditingEntryPoint;
+
+/** Responsible for registering Sandboxed Editing with the sandbox system. */
+class FSandboxedEditingSandboxIntegration : public FNoncopyable
+{
+public:
+	
+	FSandboxedEditingSandboxIntegration();
+	~FSandboxedEditingSandboxIntegration();
+	
+	const TSharedRef<FSandboxedEditingEntryPoint>& GetEntryPoint() const { return EntryPoint; }
+	const TSharedRef<FileSandboxUI::IExternalSandboxActiveViewModel>& GetExternalSandboxViewModel() const { return ExternalSandboxViewModel; }
+	
+private:
+	
+	/** Represents Sandboxed Editing as entry point so it becomes discoverable by other editor systems. */
+	const TSharedRef<FSandboxedEditingEntryPoint> EntryPoint;
+	
+	/** This view model is used to overlay a widget onto the browser when the engine is already sandboxed. */
+	const TSharedRef<FileSandboxUI::IExternalSandboxActiveViewModel> ExternalSandboxViewModel;
+};
+}
+
diff --git a/Engine/Plugins/Developer/Sandbox/SandboxedEditing/Source/SandboxedEditing/Private/Framework/Models/SandboxSystemModel.cpp b/Engine/Plugins/Developer/Sandbox/SandboxedEditing/Source/SandboxedEditing/Private/Framework/Models/SandboxSystemModel.cpp
--- a/Engine/Plugins/Developer/Sandbox/SandboxedEditing/Source/SandboxedEditing/Private/Framework/Models/SandboxSystemModel.cpp
+++ b/Engine/Plugins/Developer/Sandbox/SandboxedEditing/Source/SandboxedEditing/Private/Framework/Models/SandboxSystemModel.cpp
@@ -10,6 +10,7 @@
 #include "ISandboxInstance.h"
 #include "ISandboxManager.h"
 #include "ISandboxRepository.h"
+#include "SandboxTagging.h"
 #include "Interface/ISandboxLock.h"
 #include "Misc/Paths.h"
 #include "Templates/Function.h"
@@ -72,7 +73,8 @@ bool FSandboxSystemModel::CreateNewSandbox(const FCreateSandboxArgs& InArgs)
 		return false;
 	}
 
-	const FNewSandboxArgs Args(InArgs.Name, InArgs.Description);
+	FNewSandboxArgs Args(InArgs.Name, InArgs.Description);
+	MarkOwnedBySandboxedEditing(Args.MetaData);
 	return SandboxManager.CreateNewSandbox(Args).HasValue();
 }
 
@@ -86,7 +88,7 @@ bool FSandboxSystemModel::CanCreateNewSandbox(const FString& InName, FText* Reas
 		*Reason = FileSandboxCore::FormatDirectorySuitabilityAsReason(Suitability, InName);
 	}
 	
-	if (!CanLeaveSandbox(Reason))
+	if (HasActiveSandbox() && !CanLeaveSandbox(Reason))
 	{
 		return false;
 	}
@@ -148,7 +150,7 @@ bool FSandboxSystemModel::CanLeaveSandbox(FText* OutReason) const
 	}
 	
 	// If this happens, we should investigate why the sandbox cannot be left even though the lock said it's fine. Return a proper a reason.
-	if (ensure(!SandboxManager.CanLeaveSandbox()))
+	if (!ensure(SandboxManager.CanLeaveSandbox()))
 	{
 		SET_REASON(LOCTEXT("LeaveReason.OtherReason", "Cannot leave sandbox: no further info available."));
 		return false;
@@ -236,15 +238,13 @@ FString FSandboxSystemModel::GetActiveSandboxName() const
 {
 	using namespace FileSandboxCore;
 	IFileSandboxCoreModule& Module = IFileSandboxCoreModule::Get();
-	ISandboxInstance* SandboxInstance = Module.GetSandboxManager().GetActiveSandboxInstance();
-	return SandboxInstance ? SandboxInstance->GetInitialMetaData().Name : FString();
+	ISandboxInstance* Instance = Module.GetSandboxManager().GetActiveSandboxInstance();
+	return Instance && IsOwnedBySandboxedEditing(Instance->GetInitialMetaData()) ? Instance->GetInitialMetaData().Name : FString();
 }
 
 bool FSandboxSystemModel::HasActiveSandbox() const
 {
-	using namespace FileSandboxCore;
-	IFileSandboxCoreModule& Module = IFileSandboxCoreModule::Get();
-	return Module.GetSandboxManager().GetActiveSandboxInstance() != nullptr;
+	return !GetActiveSandboxName().IsEmpty();
 }
 
 void FSandboxSystemModel::ForEachSandbox(TFunctionRef<void(const FSandboxInfo& InRoot)> InProcess) const
@@ -255,8 +255,11 @@ void FSandboxSystemModel::ForEachSandbox(TFunctionRef<void(const FSandboxInfo& I
 	ISandboxRepository& Repository = Module.GetDefaultSandboxRepository();
 	Repository.ForEachSandbox([&InProcess](const FString& InRootPath, const FSandboxMetaInfo& MetaInfo)
 	{
-		const FFileSandboxCore_SandboxMetaData MetaData = MetaInfo.UserMetaData;
-		InProcess(FSandboxInfo(MetaData.Name, MetaData.Description, InRootPath, GetLastModified(InRootPath), MetaInfo.VersionInfo.EngineVersion));
+		const FFileSandboxCore_SandboxMetaData& MetaData = MetaInfo.UserMetaData;
+		if (IsOwnedBySandboxedEditing(MetaData))
+		{
+			InProcess(FSandboxInfo(MetaData.Name, MetaData.Description, InRootPath, GetLastModified(InRootPath), MetaInfo.VersionInfo.EngineVersion));
+		}
 		return FileSandboxCore::EBreakBehavior::Continue;
 	});
 }
@@ -277,8 +280,11 @@ TOptional<FSandboxInfo> FSandboxSystemModel::GetSandboxInfo(const FString& InSan
 	TOptional<FSandboxInfo> Result;
 	Repository.ReadMetaData(InSandboxRoot, [&InSandboxRoot, &Result](const FSandboxMetaInfo& MetaInfo)
 	{
-		const FFileSandboxCore_SandboxMetaData MetaData = MetaInfo.UserMetaData;
-		Result.Emplace(MetaData.Name, MetaData.Description, InSandboxRoot, MetaInfo.LastModified, MetaInfo.VersionInfo.EngineVersion);
+		const FFileSandboxCore_SandboxMetaData& MetaData = MetaInfo.UserMetaData;
+		if (IsOwnedBySandboxedEditing(MetaData))
+		{
+			Result.Emplace(MetaData.Name, MetaData.Description, InSandboxRoot, MetaInfo.LastModified, MetaInfo.VersionInfo.EngineVersion);
+		}
 	});
 	return Result;
 }
@@ -290,9 +296,12 @@ void FSandboxSystemModel::SetDescription(const FString& InSandboxRoot, const FSt
 	ISandboxRepository& Repository = Module.GetDefaultSandboxRepository();
 	Repository.ReadMetaData(InSandboxRoot, [&InSandboxRoot, &InDescription](const FSandboxMetaInfo& MetaInfo)
 	{
-		FFileSandboxCore_SandboxMetaData MetaData = MetaInfo.UserMetaData;
-		MetaData.Description = InDescription;
-		SaveMetaData(MetaData, InSandboxRoot);
+		if (IsOwnedBySandboxedEditing(MetaInfo.UserMetaData))
+		{
+			FFileSandboxCore_SandboxMetaData MetaData = MetaInfo.UserMetaData;
+			MetaData.Description = InDescription;
+			SaveMetaData(MetaData, InSandboxRoot);
+		}
 	});
 }
 }
diff --git a/Engine/Plugins/Developer/Sandbox/SandboxedEditing/Source/SandboxedEditing/Private/Framework/Models/SandboxTagging.cpp b/Engine/Plugins/Developer/Sandbox/SandboxedEditing/Source/SandboxedEditing/Private/Framework/Models/SandboxTagging.cpp
new file mode 100644
--- /dev/null
+++ b/Engine/Plugins/Developer/Sandbox/SandboxedEditing/Source/SandboxedEditing/Private/Framework/Models/SandboxTagging.cpp
@@ -0,0 +1,18 @@
+﻿// Copyright Epic Games, Inc. All Rights Reserved.
+
+#include "SandboxTagging.h"
+
+#include "Data/SandboxMetaData.h"
+
+namespace UE::SandboxedEditing
+{
+bool IsOwnedBySandboxedEditing(const FFileSandboxCore_SandboxMetaData& InMetaData)
+{
+	return InMetaData.Tags.Contains(SandboxedEditingTag.Resolve());
+}
+
+void MarkOwnedBySandboxedEditing(FFileSandboxCore_SandboxMetaData& InMetaData)
+{
+	InMetaData.Tags.Add(SandboxedEditingTag.Resolve());
+}
+}
diff --git a/Engine/Plugins/Developer/Sandbox/SandboxedEditing/Source/SandboxedEditing/Private/Framework/Models/SandboxTagging.h b/Engine/Plugins/Developer/Sandbox/SandboxedEditing/Source/SandboxedEditing/Private/Framework/Models/SandboxTagging.h
new file mode 100644
--- /dev/null
+++ b/Engine/Plugins/Developer/Sandbox/SandboxedEditing/Source/SandboxedEditing/Private/Framework/Models/SandboxTagging.h
@@ -0,0 +1,19 @@
+﻿// Copyright Epic Games, Inc. All Rights Reserved.
+
+#pragma once
+
+#include "UObject/NameTypes.h"
+
+struct FFileSandboxCore_SandboxMetaData;
+
+namespace UE::SandboxedEditing
+{
+/** Tag added to sandboxes created by the Sandboxed Editing plugin. */
+constexpr FLazyName SandboxedEditingTag("SandboxedEditingTag");
+
+/** @return Whether the metadata is marked to belong to the Sandboxed Editing application. */
+bool IsOwnedBySandboxedEditing(const FFileSandboxCore_SandboxMetaData& InMetaData);
+
+/** Modifies the metadata such that we know it belongs to Sandboxed Editing application. */
+void MarkOwnedBySandboxedEditing(FFileSandboxCore_SandboxMetaData& InMetaData);
+}
diff --git a/Engine/Plugins/Developer/Sandbox/SandboxedEditing/Source/SandboxedEditing/Private/Framework/SandboxedEditingApp.cpp b/Engine/Plugins/Developer/Sandbox/SandboxedEditing/Source/SandboxedEditing/Private/Framework/SandboxedEditingApp.cpp
--- a/Engine/Plugins/Developer/Sandbox/SandboxedEditing/Source/SandboxedEditing/Private/Framework/SandboxedEditingApp.cpp
+++ b/Engine/Plugins/Developer/Sandbox/SandboxedEditing/Source/SandboxedEditing/Private/Framework/SandboxedEditingApp.cpp
@@ -2,13 +2,18 @@
 
 #include "SandboxedEditingApp.h"
 
+#include "EntryPoint/SandboxedEditingEntryPoint.h"
+#include "EntryPoint/SandboxedEditingSandboxIntegration.h"
 #include "Models/SandboxSystemModel.h"
 
 namespace UE::SandboxedEditing
 {
 FSandboxedEditingApp::FSandboxedEditingApp()
 	: SandboxModel(MakeShared<FSandboxSystemModel>())
+	, SandboxIntegration(MakeShared<FSandboxedEditingSandboxIntegration>())
 	, StatusBarFeature(SandboxModel)
-	, BrowserFeature(SandboxModel)
-{}
+	, BrowserFeature(SandboxModel, SandboxIntegration->GetExternalSandboxViewModel())
+{
+	SandboxIntegration->GetEntryPoint()->OnRequestSummonUI().AddLambda([this] { BrowserFeature.SummonUI(); });
+}
 }
diff --git a/Engine/Plugins/Developer/Sandbox/SandboxedEditing/Source/SandboxedEditing/Private/Framework/SandboxedEditingApp.h b/Engine/Plugins/Developer/Sandbox/SandboxedEditing/Source/SandboxedEditing/Private/Framework/SandboxedEditingApp.h
--- a/Engine/Plugins/Developer/Sandbox/SandboxedEditing/Source/SandboxedEditing/Private/Framework/SandboxedEditingApp.h
+++ b/Engine/Plugins/Developer/Sandbox/SandboxedEditing/Source/SandboxedEditing/Private/Framework/SandboxedEditingApp.h
@@ -10,6 +10,7 @@
 namespace UE::SandboxedEditing
 {
 class FSandboxSystemModel;
+class FSandboxedEditingSandboxIntegration;
 	
 /**
  * Creates all relevant systems to deliver the sandboxed editing experience in the editor.
@@ -29,6 +30,9 @@ private:
 	 */
 	const TSharedRef<FSandboxSystemModel> SandboxModel;
 	
+	/** Responsible for registering Sandboxed Editing with the sandbox system. */
+	const TSharedRef<FSandboxedEditingSandboxIntegration> SandboxIntegration;
+	
 	/** Extends the status bar for Sandboxed Editing. */
 	const FStatusBarFeature StatusBarFeature;
 	/** Extends the editor with a nomad tab containing a browser for sandboxes. */
diff --git a/Engine/Plugins/Developer/Sandbox/SandboxedEditing/Source/SandboxedEditing/SandboxedEditing.Build.cs b/Engine/Plugins/Developer/Sandbox/SandboxedEditing/Source/SandboxedEditing/SandboxedEditing.Build.cs
--- a/Engine/Plugins/Developer/Sandbox/SandboxedEditing/Source/SandboxedEditing/SandboxedEditing.Build.cs
+++ b/Engine/Plugins/Developer/Sandbox/SandboxedEditing/Source/SandboxedEditing/SandboxedEditing.Build.cs
@@ -11,13 +11,14 @@ namespace UnrealBuildTool.Rules
 				{
 					"Core",
 					"CoreUObject",  
-					"FileSandboxCore"
 				}
 			);
 			
 			PrivateDependencyModuleNames.AddRange(
 				new string[]
 				{
+					"FileSandboxCore",
+					"FileSandboxUI",
 					"InputCore",
 					"Slate", 
 					"SlateCore", 
