Sandboxed Editing: Make the plugin visible and mark experimental.

#Synced-CL 49159615

diff --git a/Engine/Plugins/Developer/Sandbox/SandboxedEditing/SandboxedEditing.uplugin b/Engine/Plugins/Developer/Sandbox/SandboxedEditing/SandboxedEditing.uplugin
--- a/Engine/Plugins/Developer/Sandbox/SandboxedEditing/SandboxedEditing.uplugin
+++ b/Engine/Plugins/Developer/Sandbox/SandboxedEditing/SandboxedEditing.uplugin
@@ -12,8 +12,9 @@
 	"SupportURL": "",
 	"EnabledByDefault": false,
 	"CanContainContent": false,
-	"Hidden": true,
+	"Hidden": false,
 	"IsBetaVersion": false,
+	"IsExperimentalVersion": true,
 	"Installed": false,
 	"Modules": [
 		{
