Sandboxed Editing: When leaving a sandbox with in-memory changes, the user is now asked to either save or discard those changes. The reason for that is that in-memory changes cannot be sandboxed: only saved assets are sandboxed.

#Synced-CL 49159615

diff --git a/Engine/Plugins/Developer/Sandbox/SandboxedEditing/Source/SandboxedEditing/Private/Features/Browser/BrowserFeature.cpp b/Engine/Plugins/Developer/Sandbox/SandboxedEditing/Source/SandboxedEditing/Private/Features/Browser/BrowserFeature.cpp
--- a/Engine/Plugins/Developer/Sandbox/SandboxedEditing/Source/SandboxedEditing/Private/Features/Browser/BrowserFeature.cpp
+++ b/Engine/Plugins/Developer/Sandbox/SandboxedEditing/Source/SandboxedEditing/Private/Features/Browser/BrowserFeature.cpp
@@ -49,6 +49,7 @@ FBrowserFeature::FBrowserFeature(
 	, BrowserColumns(MakeBuiltInColumns())
 	, ViewModels(BrowserDetail::TransformToBehaviours(BrowserColumns), SandboxModel, InExternalSandboxViewModel)
 	, CommandBindings(MakeShared<FBrowserCommandBindings>(ViewModels))
+	, LeaveWithDirtyPackagesView(ViewModels.LeaveViewModel)
 {
 	RegisterTabSpawner();
 }
diff --git a/Engine/Plugins/Developer/Sandbox/SandboxedEditing/Source/SandboxedEditing/Private/Features/Browser/BrowserFeature.h b/Engine/Plugins/Developer/Sandbox/SandboxedEditing/Source/SandboxedEditing/Private/Features/Browser/BrowserFeature.h
--- a/Engine/Plugins/Developer/Sandbox/SandboxedEditing/Source/SandboxedEditing/Private/Features/Browser/BrowserFeature.h
+++ b/Engine/Plugins/Developer/Sandbox/SandboxedEditing/Source/SandboxedEditing/Private/Features/Browser/BrowserFeature.h
@@ -3,6 +3,7 @@
 #pragma once
 
 #include "Commands/BrowserCommandBindings.h"
+#include "NotificationViews/LeaveWithDirtyPackagesNotificationView.h"
 #include "Templates/SharedPointer.h"
 #include "Templates/UnrealTemplate.h"
 #include "ViewModels/BrowserViewModels.h"
@@ -14,10 +15,10 @@ namespace UE::SandboxedEditing
 {
 class FActiveSandboxDetailsViewModel;
 class FActiveSandboxTrackerViewModel;
-class SSandboxedEditingRoot;
 class FSandboxControlsViewModel;
 class FSandboxListViewModel;
 class FSandboxSystemModel;
+class SSandboxedEditingRoot;
 
 /** Summons a nomad tab through which the user can browse sandboxes. */
 class FBrowserFeature : public FNoncopyable
@@ -50,6 +51,9 @@ private:
 	/** Binds user commands to models.*/
 	const TSharedRef<FBrowserCommandBindings> CommandBindings;
 	
+	/** Displays editor notifications (toaster, dialogue, etc.) for when the user tries to leave the sandbox with in-memory changes. */
+	FLeaveWithDirtyPackagesNotificationView LeaveWithDirtyPackagesView;
+	
 	void RegisterTabSpawner();
 	void UnregisterTabSpawner();
 	TSharedRef<SDockTab> SpawnBrowserTab(const FSpawnTabArgs& InArgs);
diff --git a/Engine/Plugins/Developer/Sandbox/SandboxedEditing/Source/SandboxedEditing/Private/Features/Browser/Commands/BrowserCommandBindings.cpp b/Engine/Plugins/Developer/Sandbox/SandboxedEditing/Source/SandboxedEditing/Private/Features/Browser/Commands/BrowserCommandBindings.cpp
--- a/Engine/Plugins/Developer/Sandbox/SandboxedEditing/Source/SandboxedEditing/Private/Features/Browser/Commands/BrowserCommandBindings.cpp
+++ b/Engine/Plugins/Developer/Sandbox/SandboxedEditing/Source/SandboxedEditing/Private/Features/Browser/Commands/BrowserCommandBindings.cpp
@@ -12,7 +12,7 @@ namespace UE::SandboxedEditing
 {
 FBrowserCommandBindings::FBrowserCommandBindings(const FBrowserViewModels& InViewModels)
 	: ControlsViewModel(InViewModels.ControlsViewModel)
-	, ActiveSandboxViewModel(InViewModels.ActiveSandboxDetailsViewModel)
+	, LeaveSandboxViewModel(InViewModels.LeaveViewModel)
 	, CommandList(MakeShared<FUICommandList>())
 {
 	FBrowserCommands::Register();
@@ -46,12 +46,12 @@ void FBrowserCommandBindings::HandleLeaveSandbox()
 {
 	if (CanLeaveSandbox())
 	{
-		ActiveSandboxViewModel->LeaveSandbox();
+		LeaveSandboxViewModel->LeaveSandbox();
 	}
 }
 
 bool FBrowserCommandBindings::CanLeaveSandbox()
 {
-	return ActiveSandboxViewModel->CanLeaveSandbox();
+	return LeaveSandboxViewModel->CanLeaveSandbox();
 }
 }
diff --git a/Engine/Plugins/Developer/Sandbox/SandboxedEditing/Source/SandboxedEditing/Private/Features/Browser/Commands/BrowserCommandBindings.h b/Engine/Plugins/Developer/Sandbox/SandboxedEditing/Source/SandboxedEditing/Private/Features/Browser/Commands/BrowserCommandBindings.h
--- a/Engine/Plugins/Developer/Sandbox/SandboxedEditing/Source/SandboxedEditing/Private/Features/Browser/Commands/BrowserCommandBindings.h
+++ b/Engine/Plugins/Developer/Sandbox/SandboxedEditing/Source/SandboxedEditing/Private/Features/Browser/Commands/BrowserCommandBindings.h
@@ -8,6 +8,7 @@
 
 namespace UE::SandboxedEditing
 {
+class FLeaveSandboxViewModel;
 class FActiveSandboxDetailsViewModel;
 class FSandboxControlsViewModel;
 struct FBrowserViewModels;
@@ -27,7 +28,7 @@ private:
 	/** Used for editing sandbox info. */
 	const TSharedRef<FSandboxControlsViewModel> ControlsViewModel;
 	/** Used to interact with the active sandbox. */
-	const TSharedRef<FActiveSandboxDetailsViewModel> ActiveSandboxViewModel;
+	const TSharedRef<FLeaveSandboxViewModel> LeaveSandboxViewModel;
 	
 	/** Holds our command bindings */
 	const TSharedRef<FUICommandList> CommandList;
diff --git a/Engine/Plugins/Developer/Sandbox/SandboxedEditing/Source/SandboxedEditing/Private/Features/Browser/NotificationViews/LeaveWithDirtyPackagesNotificationView.cpp b/Engine/Plugins/Developer/Sandbox/SandboxedEditing/Source/SandboxedEditing/Private/Features/Browser/NotificationViews/LeaveWithDirtyPackagesNotificationView.cpp
new file mode 100644
--- /dev/null
+++ b/Engine/Plugins/Developer/Sandbox/SandboxedEditing/Source/SandboxedEditing/Private/Features/Browser/NotificationViews/LeaveWithDirtyPackagesNotificationView.cpp
@@ -0,0 +1,97 @@
+﻿// Copyright Epic Games, Inc. All Rights Reserved.
+
+#include "LeaveWithDirtyPackagesNotificationView.h"
+
+#include "Features/Browser/ViewModels/Active/LeaveSandboxViewModel.h"
+#include "Framework/Notifications/NotificationManager.h"
+#include "Widgets/Notifications/SNotificationList.h"
+
+#define LOCTEXT_NAMESPACE "FNotificationView"
+
+namespace UE::SandboxedEditing
+{
+FLeaveWithDirtyPackagesNotificationView::FLeaveWithDirtyPackagesNotificationView(const TSharedRef<FLeaveSandboxViewModel>& InLeaveViewModel)
+	: LeaveViewModel(InLeaveViewModel)
+{
+	LeaveViewModel->OnCreateLeaveWithDirtyPackageWorkflow().AddRaw(this, &FLeaveWithDirtyPackagesNotificationView::ShowNotification);
+}
+
+FLeaveWithDirtyPackagesNotificationView::~FLeaveWithDirtyPackagesNotificationView()
+{
+	LeaveViewModel->OnCreateLeaveWithDirtyPackageWorkflow().RemoveAll(this);
+	
+	if (FLeaveWithDirtyPackagesWorkflow* Workflow = LeaveViewModel->GetLeaveWithDirtyPackageWorkflow())
+	{
+		Workflow->OnWorkflowEnded().Remove(HandleWorkflowEndHandle);
+	}
+}
+
+void FLeaveWithDirtyPackagesNotificationView::ShowNotification(FLeaveWithDirtyPackagesWorkflow& InWorkflow)
+{
+	HideNotification();
+	HandleWorkflowEndHandle = InWorkflow.OnWorkflowEnded().AddLambda([this](auto){ HideNotification(); });
+	
+	FSlateNotificationManager& NotificationManager = FSlateNotificationManager::Get();
+	
+	FNotificationInfo NotificationInfo(LOCTEXT("Unsaved.Title", "In-memory changes"));
+	NotificationInfo.SubText = LOCTEXT("Unsaved.SubText", "You must either save or discard unsaved assets before leaving.");
+	NotificationInfo.bFireAndForget = false;
+	NotificationInfo.ButtonDetails = 
+	{
+		FNotificationButtonInfo(
+			LOCTEXT("Unsaved.SaveAll.Title", "Save all"),
+			LOCTEXT("Unsaved.SaveAll.ToolTip", "Save all unsaved assets"),
+			FSimpleDelegate::CreateRaw(this, &FLeaveWithDirtyPackagesNotificationView::HandleSaveAll),
+			SNotificationItem::CS_None
+			),
+		FNotificationButtonInfo(
+			LOCTEXT("Unsaved.DiscardAll.Title", "Discard"),
+			LOCTEXT("Unsaved.DiscardAll.ToolTip", "Discard all unsaved assets"),
+			FSimpleDelegate::CreateRaw(this, &FLeaveWithDirtyPackagesNotificationView::HandleRevert),
+			SNotificationItem::CS_None
+			),
+		FNotificationButtonInfo(
+			LOCTEXT("Unsaved.Cancel.Title", "Cancel"),
+			LOCTEXT("Unsaved.Cancel.ToolTip", "Does not leave the sandbox"),
+			FSimpleDelegate::CreateLambda([this]{ HideNotification(); }),
+			SNotificationItem::CS_None
+			)
+	};
+	Notification = NotificationManager.AddNotification(NotificationInfo);
+}
+
+void FLeaveWithDirtyPackagesNotificationView::HandleSaveAll() const
+{
+	if (FLeaveWithDirtyPackagesWorkflow* Workflow = LeaveViewModel->GetLeaveWithDirtyPackageWorkflow())
+	{
+		Workflow->SaveDirtyPackages();
+	}
+}
+
+void FLeaveWithDirtyPackagesNotificationView::HandleRevert() const
+{
+	if (FLeaveWithDirtyPackagesWorkflow* Workflow = LeaveViewModel->GetLeaveWithDirtyPackageWorkflow())
+	{
+		Workflow->RevertDirtyPackages();
+	}
+}
+
+void FLeaveWithDirtyPackagesNotificationView::HandleCancel() const
+{
+	if (FLeaveWithDirtyPackagesWorkflow* Workflow = LeaveViewModel->GetLeaveWithDirtyPackageWorkflow())
+	{
+		Workflow->Cancel();
+	}
+}
+
+void FLeaveWithDirtyPackagesNotificationView::HideNotification()
+{
+	if (Notification)
+	{
+		Notification->ExpireAndFadeout();
+		Notification.Reset();
+	}
+}
+}
+
+#undef LOCTEXT_NAMESPACE
\ No newline at end of file
diff --git a/Engine/Plugins/Developer/Sandbox/SandboxedEditing/Source/SandboxedEditing/Private/Features/Browser/NotificationViews/LeaveWithDirtyPackagesNotificationView.h b/Engine/Plugins/Developer/Sandbox/SandboxedEditing/Source/SandboxedEditing/Private/Features/Browser/NotificationViews/LeaveWithDirtyPackagesNotificationView.h
new file mode 100644
--- /dev/null
+++ b/Engine/Plugins/Developer/Sandbox/SandboxedEditing/Source/SandboxedEditing/Private/Features/Browser/NotificationViews/LeaveWithDirtyPackagesNotificationView.h
@@ -0,0 +1,48 @@
+﻿// Copyright Epic Games, Inc. All Rights Reserved.
+
+#pragma once
+
+#include "Delegates/IDelegateInstance.h"
+#include "Templates/SharedPointer.h"
+#include "Templates/UnrealTemplate.h"
+
+class SNotificationItem;
+
+namespace UE::SandboxedEditing
+{
+class FLeaveSandboxViewModel;
+class FLeaveWithDirtyPackagesWorkflow;
+
+/** Visualizes the notification for leaving the sandbox with dirty packages. */
+class FLeaveWithDirtyPackagesNotificationView : public FNoncopyable
+{
+public:
+	
+	explicit FLeaveWithDirtyPackagesNotificationView(const TSharedRef<FLeaveSandboxViewModel>& InLeaveViewModel);
+	~FLeaveWithDirtyPackagesNotificationView();
+	
+private:
+	
+	/** Notifies us when a FLeaveSandboxViewModel starts. */
+	const TSharedRef<FLeaveSandboxViewModel> LeaveViewModel;
+
+	/** Whether currently showing the notification for when the user is trying to leave the sandbox even though there are modified assets. */
+	TSharedPtr<SNotificationItem> Notification;
+	
+	FDelegateHandle HandleWorkflowEndHandle;
+
+	/** Show the notification for when the user is trying to leave the sandbox even though there are modified assets. */
+	void ShowNotification(FLeaveWithDirtyPackagesWorkflow& InWorkflow);
+	
+	/** Saves all dirty assets, hides the notification, and proceeds with the leave operation. */
+	void HandleSaveAll() const;
+	/** Discards all dirty assets, hides the notification, and proceeds with the leave operation. */
+	void HandleRevert() const;
+	/** Cancels the workflow. */
+	void HandleCancel() const;
+	
+	/** Hides the notification. */
+	void HideNotification();
+};
+}
+
diff --git a/Engine/Plugins/Developer/Sandbox/SandboxedEditing/Source/SandboxedEditing/Private/Features/Browser/ViewModels/Active/ActiveSandboxDetailsViewModel.cpp b/Engine/Plugins/Developer/Sandbox/SandboxedEditing/Source/SandboxedEditing/Private/Features/Browser/ViewModels/Active/ActiveSandboxDetailsViewModel.cpp
--- a/Engine/Plugins/Developer/Sandbox/SandboxedEditing/Source/SandboxedEditing/Private/Features/Browser/ViewModels/Active/ActiveSandboxDetailsViewModel.cpp
+++ b/Engine/Plugins/Developer/Sandbox/SandboxedEditing/Source/SandboxedEditing/Private/Features/Browser/ViewModels/Active/ActiveSandboxDetailsViewModel.cpp
@@ -2,27 +2,28 @@
 
 #include "ActiveSandboxDetailsViewModel.h"
 
+#include "LeaveSandboxViewModel.h"
 #include "Framework/Models/SandboxSystemModel.h"
+#include "Types/GatheredFileChanges.h"
 
 namespace UE::SandboxedEditing
 {
-FActiveSandboxDetailsViewModel::FActiveSandboxDetailsViewModel(const TSharedRef<FSandboxSystemModel>& InModel)
+FActiveSandboxDetailsViewModel::FActiveSandboxDetailsViewModel(
+	const TSharedRef<FSandboxSystemModel>& InModel, 
+	const TSharedRef<FLeaveSandboxViewModel>& InLeaveSandboxViewModel
+	)
 	: Model(InModel)
+	, LeaveSandboxViewModel(InLeaveSandboxViewModel)
 {}
 
-bool FActiveSandboxDetailsViewModel::HasActiveSandbox() const
-{
-	return Model->HasActiveSandbox();
-}
-
 void FActiveSandboxDetailsViewModel::LeaveSandbox()
 {
-	Model->LeaveSandbox();
+	LeaveSandboxViewModel->LeaveSandbox();
 }
 
 bool FActiveSandboxDetailsViewModel::CanLeaveSandbox(FText* OutReason) const
 {
-	return Model->CanLeaveSandbox(OutReason);
+	return LeaveSandboxViewModel->CanLeaveSandbox(OutReason);
 }
 
 FString FActiveSandboxDetailsViewModel::GetSandboxName() const
diff --git a/Engine/Plugins/Developer/Sandbox/SandboxedEditing/Source/SandboxedEditing/Private/Features/Browser/ViewModels/Active/ActiveSandboxDetailsViewModel.h b/Engine/Plugins/Developer/Sandbox/SandboxedEditing/Source/SandboxedEditing/Private/Features/Browser/ViewModels/Active/ActiveSandboxDetailsViewModel.h
--- a/Engine/Plugins/Developer/Sandbox/SandboxedEditing/Source/SandboxedEditing/Private/Features/Browser/ViewModels/Active/ActiveSandboxDetailsViewModel.h
+++ b/Engine/Plugins/Developer/Sandbox/SandboxedEditing/Source/SandboxedEditing/Private/Features/Browser/ViewModels/Active/ActiveSandboxDetailsViewModel.h
@@ -8,6 +8,7 @@
 
 namespace UE::SandboxedEditing
 {
+class FLeaveSandboxViewModel;
 class FSandboxSystemModel;
 
 /** Information about the active sandbox. */
@@ -15,10 +16,10 @@ class FActiveSandboxDetailsViewModel : public FNoncopyable
 {
 public:
 	
-	explicit FActiveSandboxDetailsViewModel(const TSharedRef<FSandboxSystemModel>& InModel);
-	
-	/** @return Whether currently in an active sandbox */
-	bool HasActiveSandbox() const;
+	explicit FActiveSandboxDetailsViewModel(
+		const TSharedRef<FSandboxSystemModel>& InModel, 
+		const TSharedRef<FLeaveSandboxViewModel>& InLeaveSandboxViewModel
+		);
 	
 	/** Leaves the active sandbox. */
 	void LeaveSandbox();
@@ -29,11 +30,14 @@ public:
 	FString GetSandboxName() const;
 	/** @return Path to active sandbox root. */
 	FString GetSandboxPath() const;
-
+	
 private:
 	
 	/** The model used to interact with the sandbox system. */
 	const TSharedRef<FSandboxSystemModel> Model;
+	
+	/** Handles leave logic. */
+	const TSharedRef<FLeaveSandboxViewModel> LeaveSandboxViewModel;
 };
 }
 
diff --git a/Engine/Plugins/Developer/Sandbox/SandboxedEditing/Source/SandboxedEditing/Private/Features/Browser/ViewModels/Active/LeaveSandboxViewModel.cpp b/Engine/Plugins/Developer/Sandbox/SandboxedEditing/Source/SandboxedEditing/Private/Features/Browser/ViewModels/Active/LeaveSandboxViewModel.cpp
new file mode 100644
--- /dev/null
+++ b/Engine/Plugins/Developer/Sandbox/SandboxedEditing/Source/SandboxedEditing/Private/Features/Browser/ViewModels/Active/LeaveSandboxViewModel.cpp
@@ -0,0 +1,54 @@
+﻿// Copyright Epic Games, Inc. All Rights Reserved.
+
+#include "LeaveSandboxViewModel.h"
+
+#include "Features/Browser/ViewModels/Workflows/LeaveWithDirtyPackagesWorkflow.h"
+#include "Framework/Models/SandboxSystemModel.h"
+
+namespace UE::SandboxedEditing
+{
+FLeaveSandboxViewModel::FLeaveSandboxViewModel(const TSharedRef<FSandboxSystemModel> InModel)
+	: Model(InModel)
+{}
+
+void FLeaveSandboxViewModel::LeaveSandbox()
+{
+	if (!CanLeaveSandbox())
+	{
+		return;
+	}
+	
+	if (Model->CanLeaveSandboxWithoutFurtherActions())
+	{
+		Model->LeaveSandbox();
+		return;
+	}
+	
+	const TArray<UPackage*> DirtyPackages = Model->GetInMemoryChanges();
+	if (!DirtyPackages.IsEmpty())
+	{
+		LeaveWithDirtyPackagesWorkflow = MakeUnique<FLeaveWithDirtyPackagesWorkflow>(
+			Model, 
+			FLeaveWithDirtyPackagesWorkflow::FOnCleanupWorkflow::CreateRaw(this, &FLeaveSandboxViewModel::OnDirtyPackagesHandled)
+			);
+		OnCreateLeaveWithDirtyPackageWorkflowDelegate.Broadcast(*LeaveWithDirtyPackagesWorkflow);
+	}
+	
+	// TODO: In the future, add logic for persisting active sandbox changes here.
+}
+
+bool FLeaveSandboxViewModel::CanLeaveSandbox(FText* OutReason) const
+{
+	return !LeaveWithDirtyPackagesWorkflow && Model->IsAllowedToLeaveSandbox(OutReason);
+}
+
+void FLeaveSandboxViewModel::OnDirtyPackagesHandled(FLeaveWithDirtyPackagesWorkflow::EWorkflowResult InResult)
+{
+	LeaveWithDirtyPackagesWorkflow.Reset();
+	
+	if (InResult == FLeaveWithDirtyPackagesWorkflow::EWorkflowResult::Success)
+	{
+		LeaveSandbox();
+	}
+}
+}
diff --git a/Engine/Plugins/Developer/Sandbox/SandboxedEditing/Source/SandboxedEditing/Private/Features/Browser/ViewModels/Active/LeaveSandboxViewModel.h b/Engine/Plugins/Developer/Sandbox/SandboxedEditing/Source/SandboxedEditing/Private/Features/Browser/ViewModels/Active/LeaveSandboxViewModel.h
new file mode 100644
--- /dev/null
+++ b/Engine/Plugins/Developer/Sandbox/SandboxedEditing/Source/SandboxedEditing/Private/Features/Browser/ViewModels/Active/LeaveSandboxViewModel.h
@@ -0,0 +1,49 @@
+﻿// Copyright Epic Games, Inc. All Rights Reserved.
+
+#pragma once
+
+#include "Delegates/Delegate.h"
+#include "Delegates/DelegateCombinations.h"
+#include "Features/Browser/ViewModels/Workflows/LeaveWithDirtyPackagesWorkflow.h"
+#include "Templates/SharedPointer.h"
+#include "Templates/UnrealTemplate.h"
+
+namespace UE::SandboxedEditing
+{
+class FSandboxSystemModel;
+
+DECLARE_MULTICAST_DELEGATE_OneParam(FLeaveWithDirtyPackagesWorkflowDelegate, FLeaveWithDirtyPackagesWorkflow&);
+
+/** Manages leaving a sandbox. */
+class FLeaveSandboxViewModel : public FNoncopyable
+{
+public:
+	
+	explicit FLeaveSandboxViewModel(const TSharedRef<FSandboxSystemModel> InModel);
+
+	/** Leaves the active sandbox. */
+	void LeaveSandbox();
+	/** @return Whether the sandbox can be left */
+	bool CanLeaveSandbox(FText* OutReason = nullptr) const; 
+	
+	/** @return Gets the current workflow */
+	FLeaveWithDirtyPackagesWorkflow* GetLeaveWithDirtyPackageWorkflow() const { return LeaveWithDirtyPackagesWorkflow.Get(); }
+	
+	/** Invoked when the user should be asked about what to do with dirty packages. */
+	FLeaveWithDirtyPackagesWorkflowDelegate& OnCreateLeaveWithDirtyPackageWorkflow() { return OnCreateLeaveWithDirtyPackageWorkflowDelegate; }
+	
+private:
+	
+	/** The model used to interact with the sandbox system. */
+	const TSharedRef<FSandboxSystemModel> Model;
+
+	/** Handles dirty packages when leaving. */
+	TUniquePtr<FLeaveWithDirtyPackagesWorkflow> LeaveWithDirtyPackagesWorkflow;
+	
+	/** Invoked when the user should be asked about what to do with dirty packages. */
+	FLeaveWithDirtyPackagesWorkflowDelegate OnCreateLeaveWithDirtyPackageWorkflowDelegate;
+	
+	/** Invoked when LeaveWithDirtyPackagesWorkflow finishes. */
+	void OnDirtyPackagesHandled(FLeaveWithDirtyPackagesWorkflow::EWorkflowResult InResult);
+};
+}
diff --git a/Engine/Plugins/Developer/Sandbox/SandboxedEditing/Source/SandboxedEditing/Private/Features/Browser/ViewModels/BrowserViewModels.h b/Engine/Plugins/Developer/Sandbox/SandboxedEditing/Source/SandboxedEditing/Private/Features/Browser/ViewModels/BrowserViewModels.h
--- a/Engine/Plugins/Developer/Sandbox/SandboxedEditing/Source/SandboxedEditing/Private/Features/Browser/ViewModels/BrowserViewModels.h
+++ b/Engine/Plugins/Developer/Sandbox/SandboxedEditing/Source/SandboxedEditing/Private/Features/Browser/ViewModels/BrowserViewModels.h
@@ -8,6 +8,7 @@
 #include "List/SandboxListViewModel.h"
 #include "SandboxControlsViewModel.h"
 #include "SandboxMetaDataViewModel.h"
+#include "Active/LeaveSandboxViewModel.h"
 #include "Templates/SharedPointer.h"
 
 namespace UE::FileSandboxUI
@@ -33,6 +34,8 @@ struct FBrowserViewModels
 	/** Handles changing metadata of a sandbox. */
 	const TSharedRef<FSandboxMetaDataViewModel> MetaDataViewModel;
 	
+	/** Handles leaving an active sandbox. */
+	const TSharedRef<FLeaveSandboxViewModel> LeaveViewModel;
 	/** Used for widgets to switch states when the active sandbox changes. */
 	const TSharedRef<FActiveSandboxTrackerViewModel> ActiveSandboxTrackerViewModel;
 	/** Used for widgets that display info about the currently active sandbox. */
@@ -54,8 +57,9 @@ struct FBrowserViewModels
 		, ListViewModel(MakeShared<FSandboxListViewModel>(InModel, FilterSandboxViewModel, MoveTemp(InColumns)))
 		, ControlsViewModel(MakeShared<FSandboxControlsViewModel>(InModel))
 		, MetaDataViewModel(MakeShared<FSandboxMetaDataViewModel>(InModel))
+		, LeaveViewModel(MakeShared<FLeaveSandboxViewModel>(InModel))
 		, ActiveSandboxTrackerViewModel(MakeShared<FActiveSandboxTrackerViewModel>(InModel))
-		, ActiveSandboxDetailsViewModel(MakeShared<FActiveSandboxDetailsViewModel>(InModel))
+		, ActiveSandboxDetailsViewModel(MakeShared<FActiveSandboxDetailsViewModel>(InModel, LeaveViewModel))
 		, ExternalSandboxViewModel(InExternalSandboxViewModel)
 	{}
 };
diff --git a/Engine/Plugins/Developer/Sandbox/SandboxedEditing/Source/SandboxedEditing/Private/Features/Browser/ViewModels/Workflows/LeaveWithDirtyPackagesWorkflow.cpp b/Engine/Plugins/Developer/Sandbox/SandboxedEditing/Source/SandboxedEditing/Private/Features/Browser/ViewModels/Workflows/LeaveWithDirtyPackagesWorkflow.cpp
new file mode 100644
--- /dev/null
+++ b/Engine/Plugins/Developer/Sandbox/SandboxedEditing/Source/SandboxedEditing/Private/Features/Browser/ViewModels/Workflows/LeaveWithDirtyPackagesWorkflow.cpp
@@ -0,0 +1,38 @@
+﻿// Copyright Epic Games, Inc. All Rights Reserved.
+
+#include "LeaveWithDirtyPackagesWorkflow.h"
+
+#include "Framework/Models/SandboxSystemModel.h"
+
+namespace UE::SandboxedEditing
+{
+FLeaveWithDirtyPackagesWorkflow::FLeaveWithDirtyPackagesWorkflow(
+	const TSharedRef<FSandboxSystemModel>& InSandboxModel, FOnCleanupWorkflow InCleanupWorkflowDelegate
+	)
+	: SandboxModel(InSandboxModel)
+	, CleanupWorkflowDelegate(MoveTemp(InCleanupWorkflowDelegate))
+{}
+
+void FLeaveWithDirtyPackagesWorkflow::RevertDirtyPackages()
+{
+	const bool bSuccess = SandboxModel->RevertDirtyPackages();
+	BroadcastEndWorkflowDelegates(bSuccess ? EWorkflowResult::Success : EWorkflowResult::Error);
+}
+
+void FLeaveWithDirtyPackagesWorkflow::SaveDirtyPackages()
+{
+	const bool bSuccess = SandboxModel->SaveDirtyPackages();
+	BroadcastEndWorkflowDelegates(bSuccess ? EWorkflowResult::Success : EWorkflowResult::Error);
+}
+
+void FLeaveWithDirtyPackagesWorkflow::Cancel()
+{
+	BroadcastEndWorkflowDelegates(EWorkflowResult::Cancelled);
+}
+
+void FLeaveWithDirtyPackagesWorkflow::BroadcastEndWorkflowDelegates(EWorkflowResult InResult)
+{
+	OnWorkflowEndedDelegate.Broadcast(InResult);
+	CleanupWorkflowDelegate.ExecuteIfBound(InResult);
+}
+}
diff --git a/Engine/Plugins/Developer/Sandbox/SandboxedEditing/Source/SandboxedEditing/Private/Features/Browser/ViewModels/Workflows/LeaveWithDirtyPackagesWorkflow.h b/Engine/Plugins/Developer/Sandbox/SandboxedEditing/Source/SandboxedEditing/Private/Features/Browser/ViewModels/Workflows/LeaveWithDirtyPackagesWorkflow.h
new file mode 100644
--- /dev/null
+++ b/Engine/Plugins/Developer/Sandbox/SandboxedEditing/Source/SandboxedEditing/Private/Features/Browser/ViewModels/Workflows/LeaveWithDirtyPackagesWorkflow.h
@@ -0,0 +1,47 @@
+﻿// Copyright Epic Games, Inc. All Rights Reserved.
+
+#pragma once
+
+#include "Delegates/Delegate.h"
+#include "Templates/UnrealTemplate.h"
+
+namespace UE::SandboxedEditing
+{
+class FSandboxSystemModel;
+
+class FLeaveWithDirtyPackagesWorkflow : public FNoncopyable
+{
+public:
+	
+	enum class EWorkflowResult { Success, Error, Cancelled };
+	DECLARE_DELEGATE_OneParam(FOnCleanupWorkflow, EWorkflowResult)
+	DECLARE_MULTICAST_DELEGATE_OneParam(FOnEndWorkflow, EWorkflowResult)
+	
+	explicit FLeaveWithDirtyPackagesWorkflow(const TSharedRef<FSandboxSystemModel>& InSandboxModel, FOnCleanupWorkflow InCleanupWorkflowDelegate);
+	
+	/** Discards all dirty packages, and reloads them. */
+	void RevertDirtyPackages();
+	
+	/** Saves all dirty packages. */
+	void SaveDirtyPackages();
+	
+	/** Cancels this workflow */
+	void Cancel();
+	
+	/** Invoked when this workflow ends. */
+	FOnEndWorkflow& OnWorkflowEnded() { return OnWorkflowEndedDelegate; } 
+	
+private:
+	
+	/** The shared sandbox model. */
+	const TSharedRef<FSandboxSystemModel> SandboxModel;
+	
+	/** Cleanup delegate invoked after OnWorkflowEndedDelegate. */
+	const FOnCleanupWorkflow CleanupWorkflowDelegate;
+	/** Invoked when this workflow ends. */
+	FOnEndWorkflow OnWorkflowEndedDelegate;
+	
+	void BroadcastEndWorkflowDelegates(EWorkflowResult InResult);
+};
+}
+
diff --git a/Engine/Plugins/Developer/Sandbox/SandboxedEditing/Source/SandboxedEditing/Private/Features/StatusBar/Commands/StatusBarCommandMappings.cpp b/Engine/Plugins/Developer/Sandbox/SandboxedEditing/Source/SandboxedEditing/Private/Features/StatusBar/Commands/StatusBarCommandMappings.cpp
--- a/Engine/Plugins/Developer/Sandbox/SandboxedEditing/Source/SandboxedEditing/Private/Features/StatusBar/Commands/StatusBarCommandMappings.cpp
+++ b/Engine/Plugins/Developer/Sandbox/SandboxedEditing/Source/SandboxedEditing/Private/Features/StatusBar/Commands/StatusBarCommandMappings.cpp
@@ -26,7 +26,7 @@ FStatusBarCommandMappings::FStatusBarCommandMappings(const TSharedRef<FSandboxSy
 	CommandList->MapAction(
 		Commands.LeaveSandbox,
 		FExecuteAction::CreateRaw(this, &FStatusBarCommandMappings::HandleLeaveSandbox),
-		FCanExecuteAction::CreateSP(InSandboxViewModel, &FSandboxSystemModel::CanLeaveSandbox)
+		FCanExecuteAction::CreateSP(InSandboxViewModel, &FSandboxSystemModel::IsAllowedToLeaveSandbox)
 		);
 
 	CommandList->MapAction(
diff --git a/Engine/Plugins/Developer/Sandbox/SandboxedEditing/Source/SandboxedEditing/Private/Features/StatusBar/Widgets/SSandboxStatusBarWidget.cpp b/Engine/Plugins/Developer/Sandbox/SandboxedEditing/Source/SandboxedEditing/Private/Features/StatusBar/Widgets/SSandboxStatusBarWidget.cpp
--- a/Engine/Plugins/Developer/Sandbox/SandboxedEditing/Source/SandboxedEditing/Private/Features/StatusBar/Widgets/SSandboxStatusBarWidget.cpp
+++ b/Engine/Plugins/Developer/Sandbox/SandboxedEditing/Source/SandboxedEditing/Private/Features/StatusBar/Widgets/SSandboxStatusBarWidget.cpp
@@ -72,7 +72,7 @@ TSharedRef<SWidget> SSandboxStatusBarWidget::CreateStatusBarMenu()
 		const TAttribute<FText> LeaveToolTipOverride = TAttribute<FText>::CreateLambda([Model = ViewModel]
 		{
 			FText Reason;
-			return Model->CanLeaveSandbox(&Reason) ? FStatusBarCommands::Get().LeaveSandbox->GetDescription() : Reason;
+			return Model->IsAllowedToLeaveSandbox(&Reason) ? FStatusBarCommands::Get().LeaveSandbox->GetDescription() : Reason;
 		});
 		
 		Section.AddMenuEntry(FStatusBarCommands::Get().OpenCreateNewSandboxDialog);
diff --git a/Engine/Plugins/Developer/Sandbox/SandboxedEditing/Source/SandboxedEditing/Private/Framework/Models/SandboxSystemModel.cpp b/Engine/Plugins/Developer/Sandbox/SandboxedEditing/Source/SandboxedEditing/Private/Framework/Models/SandboxSystemModel.cpp
--- a/Engine/Plugins/Developer/Sandbox/SandboxedEditing/Source/SandboxedEditing/Private/Framework/Models/SandboxSystemModel.cpp
+++ b/Engine/Plugins/Developer/Sandbox/SandboxedEditing/Source/SandboxedEditing/Private/Framework/Models/SandboxSystemModel.cpp
@@ -2,6 +2,7 @@
 
 #include "SandboxSystemModel.h"
 
+#include "AssetViewUtils.h"
 #include "Containers/UnrealString.h"
 #include "Data/SandboxMetaData.h"
 #include "Framework/Models/CreateSandboxArgs.h"
@@ -20,6 +21,7 @@
 #include "Types/Manager/DeleteSandboxResult.h"
 #include "Types/Manager/LoadSandboxByDirectoryArgs.h"
 #include "Types/Manager/NewSandboxArgs.h"
+#include "Utils/PackageSandboxUtils.h"
 
 #define LOCTEXT_NAMESPACE "FSandboxSystemModel"
 
@@ -88,7 +90,7 @@ bool FSandboxSystemModel::CanCreateNewSandbox(const FString& InName, FText* Reas
 		*Reason = FileSandboxCore::FormatDirectorySuitabilityAsReason(Suitability, InName);
 	}
 	
-	if (HasActiveSandbox() && !CanLeaveSandbox(Reason))
+	if (HasActiveSandbox() && !IsAllowedToLeaveSandbox(Reason))
 	{
 		return false;
 	}
@@ -131,7 +133,7 @@ void FSandboxSystemModel::LeaveSandbox() const
 }
 
 #define SET_REASON(ReasonText) if (OutReason){ *OutReason = ReasonText; }
-bool FSandboxSystemModel::CanLeaveSandbox(FText* OutReason) const
+bool FSandboxSystemModel::IsAllowedToLeaveSandbox(FText* OutReason) const
 {
 	using namespace FileSandboxCore;
 	IFileSandboxCoreModule& Module = IFileSandboxCoreModule::Get();
@@ -146,20 +148,38 @@ bool FSandboxSystemModel::CanLeaveSandbox(FText* OutReason) const
 	if (ISandboxLock* Lock = SandboxManager.GetActiveLock()
 		; Lock && !Lock->CanLeaveSandbox(OutReason))
 	{
+		SET_REASON(LOCTEXT("LeaveReason.Locked", "The engine is locked to be sandboxed."));
 		return false;
 	}
 	
-	// If this happens, we should investigate why the sandbox cannot be left even though the lock said it's fine. Return a proper a reason.
-	if (!ensure(SandboxManager.CanLeaveSandbox()))
-	{
-		SET_REASON(LOCTEXT("LeaveReason.OtherReason", "Cannot leave sandbox: no further info available."));
-		return false;
-	}
-	
-	return true;
+	const ELeaveSandboxErrorCode LeaveReason = SandboxManager.CanLeaveSandboxWithReason();
+	return LeaveReason == ELeaveSandboxErrorCode::InMemoryChanges  // The user is able to leave the sandbox - the app will handle the dirty packages.
+		|| LeaveReason == ELeaveSandboxErrorCode::Success;
 }
 #undef SET_REASON
 
+bool FSandboxSystemModel::CanLeaveSandboxWithoutFurtherActions() const
+{
+	return IsAllowedToLeaveSandbox() && GetInMemoryChanges().IsEmpty();
+}
+
+TArray<UPackage*> FSandboxSystemModel::GetInMemoryChanges() const
+{
+	return FileSandboxCore::GetDirtyPackages();
+}
+
+bool FSandboxSystemModel::RevertDirtyPackages()
+{
+	return HasActiveSandbox() && FileSandboxCore::HotReloadPackages(
+		FileSandboxCore::GetDirtyPackages()
+		);
+}
+
+bool FSandboxSystemModel::SaveDirtyPackages()
+{
+	return HasActiveSandbox() && AssetViewUtils::SaveDirtyPackages();
+}
+
 void FSandboxSystemModel::DeleteSandbox(const FString& InSandboxRoot) const
 {
 	if (CanDeleteSandbox(InSandboxRoot))
diff --git a/Engine/Plugins/Developer/Sandbox/SandboxedEditing/Source/SandboxedEditing/Private/Framework/Models/SandboxSystemModel.h b/Engine/Plugins/Developer/Sandbox/SandboxedEditing/Source/SandboxedEditing/Private/Framework/Models/SandboxSystemModel.h
--- a/Engine/Plugins/Developer/Sandbox/SandboxedEditing/Source/SandboxedEditing/Private/Framework/Models/SandboxSystemModel.h
+++ b/Engine/Plugins/Developer/Sandbox/SandboxedEditing/Source/SandboxedEditing/Private/Framework/Models/SandboxSystemModel.h
@@ -8,6 +8,7 @@
 
 namespace UE::FileSandboxCore
 {
+struct FGatheredFileChanges;
 class ISandboxInstance;
 }
 
@@ -19,7 +20,7 @@ struct FCreateSandboxArgs;
 struct FSandboxInfo;
 	
 /**
- * View model for creating, loading, and querying the active sandbox.
+ * Model for creating, loading, and querying the active sandbox.
  * 
  * Bridges the sandbox system without exposing it directly.
  * Members added here should have direct corresponding UI elements.
@@ -46,9 +47,18 @@ public:
 	/** Leaves the current sandbox. */
 	void LeaveSandbox() const;
 	/** @return Whether the sandbox can be left (false if there is no active one). */
-	bool CanLeaveSandbox() const { return CanLeaveSandbox(nullptr); }
+	bool IsAllowedToLeaveSandbox() const { return IsAllowedToLeaveSandbox(nullptr); }
 	/** @return Whether the sandbox can be left (false if there is no active one). */
-	bool CanLeaveSandbox(FText* OutReason) const;
+	bool IsAllowedToLeaveSandbox(FText* OutReason) const;
+	/** @return Whether the engine is in a valid state to leave the sandbox, i.e. no other clean up actions need to be performed. */
+	bool CanLeaveSandboxWithoutFurtherActions() const;
+	
+	/** @return The in-memory changes preventing the sandbox from being left. */
+	TArray<UPackage*> GetInMemoryChanges() const;
+	/** Reverts in-memory changes on the given packages, and performs related actions, such as reloading the underlying asset editors, etc. */
+	bool RevertDirtyPackages();
+	/** Saves the in-memory changes of the given packages, so the packages are no longer dirty. */
+	bool SaveDirtyPackages();
 	
 	/** Deletes the sandbox. */
 	void DeleteSandbox(const FString& InSandboxRoot) const;
@@ -71,7 +81,7 @@ public:
 	FString GetActiveSandboxName() const;
 	/** @return Whether there is an active sandbox */
 	bool HasActiveSandbox() const;
-
+	
 	/** Enumerates all known sandboxes. */
 	void ForEachSandbox(TFunctionRef<void(const FSandboxInfo& InRoot)> InProcess) const;
 	/** @return The sandboxes we know of. */
diff --git a/Engine/Plugins/Developer/Sandbox/SandboxedEditing/Source/SandboxedEditing/SandboxedEditing.Build.cs b/Engine/Plugins/Developer/Sandbox/SandboxedEditing/Source/SandboxedEditing/SandboxedEditing.Build.cs
--- a/Engine/Plugins/Developer/Sandbox/SandboxedEditing/Source/SandboxedEditing/SandboxedEditing.Build.cs
+++ b/Engine/Plugins/Developer/Sandbox/SandboxedEditing/Source/SandboxedEditing/SandboxedEditing.Build.cs
@@ -17,6 +17,7 @@ namespace UnrealBuildTool.Rules
 			PrivateDependencyModuleNames.AddRange(
 				new string[]
 				{
+					"AssetTools",
 					"FileSandboxCore",
 					"FileSandboxUI",
 					"InputCore",
@@ -24,7 +25,7 @@ namespace UnrealBuildTool.Rules
 					"SlateCore", 
 					"ToolMenus", 
 					"ToolWidgets", 
-					"WorkspaceMenuStructure",
+					"WorkspaceMenuStructure", 
 				}
 			);
 		}
