Sandboxed Editing: Create ViewModels/Leaving directory and rename some of the views & workflows to be more clear (AskUserAboutX).

#Synced-CL 49159615

diff --git a/Engine/Plugins/Developer/Sandbox/SandboxedEditing/Source/SandboxedEditing/Private/Features/Browser/BrowserFeature.h b/Engine/Plugins/Developer/Sandbox/SandboxedEditing/Source/SandboxedEditing/Private/Features/Browser/BrowserFeature.h
--- a/Engine/Plugins/Developer/Sandbox/SandboxedEditing/Source/SandboxedEditing/Private/Features/Browser/BrowserFeature.h
+++ b/Engine/Plugins/Developer/Sandbox/SandboxedEditing/Source/SandboxedEditing/Private/Features/Browser/BrowserFeature.h
@@ -3,10 +3,10 @@
 #pragma once
 
 #include "Commands/BrowserCommandBindings.h"
-#include "NotificationViews/LeaveWithDirtyPackagesNotificationView.h"
 #include "Templates/SharedPointer.h"
 #include "Templates/UnrealTemplate.h"
 #include "ViewModels/BrowserViewModels.h"
+#include "Views/AskUserAboutDirtyPackagesView.h"
 
 class SDockTab;
 class FSpawnTabArgs;
@@ -52,7 +52,7 @@ private:
 	const TSharedRef<FBrowserCommandBindings> CommandBindings;
 	
 	/** Displays editor notifications (toaster, dialogue, etc.) for when the user tries to leave the sandbox with in-memory changes. */
-	FLeaveWithDirtyPackagesNotificationView LeaveWithDirtyPackagesView;
+	FAskUserAboutDirtyPackagesView LeaveWithDirtyPackagesView;
 	
 	void RegisterTabSpawner();
 	void UnregisterTabSpawner();
diff --git a/Engine/Plugins/Developer/Sandbox/SandboxedEditing/Source/SandboxedEditing/Private/Features/Browser/NotificationViews/LeaveWithDirtyPackagesNotificationView.cpp b/Engine/Plugins/Developer/Sandbox/SandboxedEditing/Source/SandboxedEditing/Private/Features/Browser/NotificationViews/LeaveWithDirtyPackagesNotificationView.cpp
deleted file mode 100644
diff --git a/Engine/Plugins/Developer/Sandbox/SandboxedEditing/Source/SandboxedEditing/Private/Features/Browser/NotificationViews/LeaveWithDirtyPackagesNotificationView.h b/Engine/Plugins/Developer/Sandbox/SandboxedEditing/Source/SandboxedEditing/Private/Features/Browser/NotificationViews/LeaveWithDirtyPackagesNotificationView.h
deleted file mode 100644
diff --git a/Engine/Plugins/Developer/Sandbox/SandboxedEditing/Source/SandboxedEditing/Private/Features/Browser/ViewModels/Active/ActiveSandboxDetailsViewModel.cpp b/Engine/Plugins/Developer/Sandbox/SandboxedEditing/Source/SandboxedEditing/Private/Features/Browser/ViewModels/Active/ActiveSandboxDetailsViewModel.cpp
--- a/Engine/Plugins/Developer/Sandbox/SandboxedEditing/Source/SandboxedEditing/Private/Features/Browser/ViewModels/Active/ActiveSandboxDetailsViewModel.cpp
+++ b/Engine/Plugins/Developer/Sandbox/SandboxedEditing/Source/SandboxedEditing/Private/Features/Browser/ViewModels/Active/ActiveSandboxDetailsViewModel.cpp
@@ -2,9 +2,8 @@
 
 #include "ActiveSandboxDetailsViewModel.h"
 
-#include "LeaveSandboxViewModel.h"
+#include "Features/Browser/ViewModels/Leaving/LeaveSandboxViewModel.h"
 #include "Framework/Models/SandboxSystemModel.h"
-#include "Types/GatheredFileChanges.h"
 
 namespace UE::SandboxedEditing
 {
diff --git a/Engine/Plugins/Developer/Sandbox/SandboxedEditing/Source/SandboxedEditing/Private/Features/Browser/ViewModels/Active/LeaveSandboxViewModel.cpp b/Engine/Plugins/Developer/Sandbox/SandboxedEditing/Source/SandboxedEditing/Private/Features/Browser/ViewModels/Active/LeaveSandboxViewModel.cpp
deleted file mode 100644
diff --git a/Engine/Plugins/Developer/Sandbox/SandboxedEditing/Source/SandboxedEditing/Private/Features/Browser/ViewModels/Active/LeaveSandboxViewModel.h b/Engine/Plugins/Developer/Sandbox/SandboxedEditing/Source/SandboxedEditing/Private/Features/Browser/ViewModels/Active/LeaveSandboxViewModel.h
deleted file mode 100644
diff --git a/Engine/Plugins/Developer/Sandbox/SandboxedEditing/Source/SandboxedEditing/Private/Features/Browser/ViewModels/BrowserViewModels.h b/Engine/Plugins/Developer/Sandbox/SandboxedEditing/Source/SandboxedEditing/Private/Features/Browser/ViewModels/BrowserViewModels.h
--- a/Engine/Plugins/Developer/Sandbox/SandboxedEditing/Source/SandboxedEditing/Private/Features/Browser/ViewModels/BrowserViewModels.h
+++ b/Engine/Plugins/Developer/Sandbox/SandboxedEditing/Source/SandboxedEditing/Private/Features/Browser/ViewModels/BrowserViewModels.h
@@ -4,11 +4,11 @@
 
 #include "Active/ActiveSandboxDetailsViewModel.h"
 #include "Active/ActiveSandboxTrackerViewModel.h"
+#include "Leaving/LeaveSandboxViewModel.h"
 #include "List/FilterSandboxViewModel.h"
 #include "List/SandboxListViewModel.h"
 #include "SandboxControlsViewModel.h"
 #include "SandboxMetaDataViewModel.h"
-#include "Active/LeaveSandboxViewModel.h"
 #include "Templates/SharedPointer.h"
 
 namespace UE::FileSandboxUI
diff --git a/Engine/Plugins/Developer/Sandbox/SandboxedEditing/Source/SandboxedEditing/Private/Features/Browser/ViewModels/Leaving/AskUserAboutDirtyPackagesWorkflow.cpp b/Engine/Plugins/Developer/Sandbox/SandboxedEditing/Source/SandboxedEditing/Private/Features/Browser/ViewModels/Leaving/AskUserAboutDirtyPackagesWorkflow.cpp
new file mode 100644
--- /dev/null
+++ b/Engine/Plugins/Developer/Sandbox/SandboxedEditing/Source/SandboxedEditing/Private/Features/Browser/ViewModels/Leaving/AskUserAboutDirtyPackagesWorkflow.cpp
@@ -0,0 +1,38 @@
+﻿// Copyright Epic Games, Inc. All Rights Reserved.
+
+#include "AskUserAboutDirtyPackagesWorkflow.h"
+
+#include "Framework/Models/SandboxSystemModel.h"
+
+namespace UE::SandboxedEditing
+{
+FAskUserAboutDirtyPackagesWorkflow::FAskUserAboutDirtyPackagesWorkflow(
+	const TSharedRef<FSandboxSystemModel>& InSandboxModel, FOnCleanupWorkflow InCleanupWorkflowDelegate
+	)
+	: SandboxModel(InSandboxModel)
+	, CleanupWorkflowDelegate(MoveTemp(InCleanupWorkflowDelegate))
+{}
+
+void FAskUserAboutDirtyPackagesWorkflow::RevertDirtyPackages()
+{
+	const bool bSuccess = SandboxModel->RevertDirtyPackages();
+	BroadcastEndWorkflowDelegates(bSuccess ? EWorkflowResult::Success : EWorkflowResult::Error);
+}
+
+void FAskUserAboutDirtyPackagesWorkflow::SaveDirtyPackages()
+{
+	const bool bSuccess = SandboxModel->SaveDirtyPackages();
+	BroadcastEndWorkflowDelegates(bSuccess ? EWorkflowResult::Success : EWorkflowResult::Error);
+}
+
+void FAskUserAboutDirtyPackagesWorkflow::Cancel()
+{
+	BroadcastEndWorkflowDelegates(EWorkflowResult::Cancelled);
+}
+
+void FAskUserAboutDirtyPackagesWorkflow::BroadcastEndWorkflowDelegates(EWorkflowResult InResult)
+{
+	OnWorkflowEndedDelegate.Broadcast(InResult);
+	CleanupWorkflowDelegate.ExecuteIfBound(InResult);
+}
+}
diff --git a/Engine/Plugins/Developer/Sandbox/SandboxedEditing/Source/SandboxedEditing/Private/Features/Browser/ViewModels/Leaving/AskUserAboutDirtyPackagesWorkflow.h b/Engine/Plugins/Developer/Sandbox/SandboxedEditing/Source/SandboxedEditing/Private/Features/Browser/ViewModels/Leaving/AskUserAboutDirtyPackagesWorkflow.h
new file mode 100644
--- /dev/null
+++ b/Engine/Plugins/Developer/Sandbox/SandboxedEditing/Source/SandboxedEditing/Private/Features/Browser/ViewModels/Leaving/AskUserAboutDirtyPackagesWorkflow.h
@@ -0,0 +1,48 @@
+﻿// Copyright Epic Games, Inc. All Rights Reserved.
+
+#pragma once
+
+#include "Delegates/Delegate.h"
+#include "Templates/UnrealTemplate.h"
+
+namespace UE::SandboxedEditing
+{
+class FSandboxSystemModel;
+
+/** Workflow for asking the user what to do about dirty packages, i.e. when leaving the sandbox. */
+class FAskUserAboutDirtyPackagesWorkflow : public FNoncopyable
+{
+public:
+	
+	enum class EWorkflowResult { Success, Error, Cancelled };
+	DECLARE_DELEGATE_OneParam(FOnCleanupWorkflow, EWorkflowResult)
+	DECLARE_MULTICAST_DELEGATE_OneParam(FOnEndWorkflow, EWorkflowResult)
+	
+	explicit FAskUserAboutDirtyPackagesWorkflow(const TSharedRef<FSandboxSystemModel>& InSandboxModel, FOnCleanupWorkflow InCleanupWorkflowDelegate);
+	
+	/** Discards all dirty packages, and reloads them. */
+	void RevertDirtyPackages();
+	
+	/** Saves all dirty packages. */
+	void SaveDirtyPackages();
+	
+	/** Cancels this workflow */
+	void Cancel();
+	
+	/** Invoked when this workflow ends. */
+	FOnEndWorkflow& OnWorkflowEnded() { return OnWorkflowEndedDelegate; } 
+	
+private:
+	
+	/** The shared sandbox model. */
+	const TSharedRef<FSandboxSystemModel> SandboxModel;
+	
+	/** Cleanup delegate invoked after OnWorkflowEndedDelegate. */
+	const FOnCleanupWorkflow CleanupWorkflowDelegate;
+	/** Invoked when this workflow ends. */
+	FOnEndWorkflow OnWorkflowEndedDelegate;
+	
+	void BroadcastEndWorkflowDelegates(EWorkflowResult InResult);
+};
+}
+
diff --git a/Engine/Plugins/Developer/Sandbox/SandboxedEditing/Source/SandboxedEditing/Private/Features/Browser/ViewModels/Leaving/LeaveSandboxViewModel.cpp b/Engine/Plugins/Developer/Sandbox/SandboxedEditing/Source/SandboxedEditing/Private/Features/Browser/ViewModels/Leaving/LeaveSandboxViewModel.cpp
new file mode 100644
--- /dev/null
+++ b/Engine/Plugins/Developer/Sandbox/SandboxedEditing/Source/SandboxedEditing/Private/Features/Browser/ViewModels/Leaving/LeaveSandboxViewModel.cpp
@@ -0,0 +1,56 @@
+﻿// Copyright Epic Games, Inc. All Rights Reserved.
+
+#include "LeaveSandboxViewModel.h"
+
+#include "Framework/Models/SandboxSystemModel.h"
+#include "Types/GatheredFileChanges.h"
+
+namespace UE::SandboxedEditing
+{
+FLeaveSandboxViewModel::FLeaveSandboxViewModel(
+	const TSharedRef<FSandboxSystemModel> InModel
+	)
+	: Model(InModel)
+{}
+
+void FLeaveSandboxViewModel::LeaveSandbox()
+{
+	if (!CanLeaveSandbox())
+	{
+		return;
+	}
+	
+	if (Model->CanLeaveSandboxWithoutFurtherActions())
+	{
+		Model->LeaveSandbox();
+		return;
+	}
+	
+	const TArray<UPackage*> DirtyPackages = Model->GetInMemoryChanges();
+	if (!DirtyPackages.IsEmpty())
+	{
+		LeaveWithDirtyPackagesWorkflow = MakeUnique<FAskUserAboutDirtyPackagesWorkflow>(
+			Model, 
+			FAskUserAboutDirtyPackagesWorkflow::FOnCleanupWorkflow::CreateRaw(this, &FLeaveSandboxViewModel::OnDirtyPackagesHandled)
+			);
+		OnCreateLeaveWithDirtyPackageWorkflowDelegate.Broadcast(*LeaveWithDirtyPackagesWorkflow);
+		return;
+	}
+}
+
+bool FLeaveSandboxViewModel::CanLeaveSandbox(FText* OutReason) const
+{
+	return !LeaveWithDirtyPackagesWorkflow && Model->IsAllowedToLeaveSandbox(OutReason);
+}
+
+void FLeaveSandboxViewModel::OnDirtyPackagesHandled(FAskUserAboutDirtyPackagesWorkflow::EWorkflowResult InResult)
+{
+	LeaveWithDirtyPackagesWorkflow.Reset();
+	
+	if (InResult == FAskUserAboutDirtyPackagesWorkflow::EWorkflowResult::Success)
+	{
+		// Proceed with the leave operation.
+		LeaveSandbox();
+	}
+}
+}
diff --git a/Engine/Plugins/Developer/Sandbox/SandboxedEditing/Source/SandboxedEditing/Private/Features/Browser/ViewModels/Leaving/LeaveSandboxViewModel.h b/Engine/Plugins/Developer/Sandbox/SandboxedEditing/Source/SandboxedEditing/Private/Features/Browser/ViewModels/Leaving/LeaveSandboxViewModel.h
new file mode 100644
--- /dev/null
+++ b/Engine/Plugins/Developer/Sandbox/SandboxedEditing/Source/SandboxedEditing/Private/Features/Browser/ViewModels/Leaving/LeaveSandboxViewModel.h
@@ -0,0 +1,61 @@
+﻿// Copyright Epic Games, Inc. All Rights Reserved.
+
+#pragma once
+
+#include "AskUserAboutDirtyPackagesWorkflow.h"
+#include "Delegates/Delegate.h"
+#include "Delegates/DelegateCombinations.h"
+#include "Templates/SharedPointer.h"
+#include "Templates/UnrealTemplate.h"
+
+namespace UE::SandboxedEditing
+{
+class FSandboxSystemModel;
+enum class EPersistWorkflowAction : uint8;
+
+DECLARE_MULTICAST_DELEGATE_OneParam(FAskUserAboutDirtyPackagesWorkflowDelegate, FAskUserAboutDirtyPackagesWorkflow&);
+
+/**
+ * Handles the userflow of leaving a sandbox. 
+ * 
+ * All steps that may be performed are:
+ * 1. Handle any in-memory changes to packages: either save all, or discard all.
+ * 2. Ask the user to persist files changes, if any have been made.
+ * 3. Persist files the user has selected.
+ * 4. Leave the sandbox.
+ */
+class FLeaveSandboxViewModel : public FNoncopyable
+{
+public:
+	
+	explicit FLeaveSandboxViewModel(
+		const TSharedRef<FSandboxSystemModel> InModel
+		);
+
+	/** Leaves the active sandbox. */
+	void LeaveSandbox();
+	/** @return Whether the sandbox can be left */
+	bool CanLeaveSandbox(FText* OutReason = nullptr) const; 
+	
+	/** @return Gets the current workflow */
+	FAskUserAboutDirtyPackagesWorkflow* GetLeaveWithDirtyPackageWorkflow() const { return LeaveWithDirtyPackagesWorkflow.Get(); }
+	
+	/** Invoked when the user should be asked about what to do with dirty packages. */
+	FAskUserAboutDirtyPackagesWorkflowDelegate& OnCreateLeaveWithDirtyPackageWorkflow() { return OnCreateLeaveWithDirtyPackageWorkflowDelegate; }
+	
+private:
+	
+	/** The model used to interact with the sandbox system. */
+	const TSharedRef<FSandboxSystemModel> Model;
+
+
+	/** Handles dirty packages when leaving. */
+	TUniquePtr<FAskUserAboutDirtyPackagesWorkflow> LeaveWithDirtyPackagesWorkflow;
+	
+	/** Invoked when the user should be asked about what to do with dirty packages. */
+	FAskUserAboutDirtyPackagesWorkflowDelegate OnCreateLeaveWithDirtyPackageWorkflowDelegate;
+	
+	/** Invoked when LeaveWithDirtyPackagesWorkflow finishes. */
+	void OnDirtyPackagesHandled(FAskUserAboutDirtyPackagesWorkflow::EWorkflowResult InResult);
+};
+}
diff --git a/Engine/Plugins/Developer/Sandbox/SandboxedEditing/Source/SandboxedEditing/Private/Features/Browser/ViewModels/Workflows/LeaveWithDirtyPackagesWorkflow.cpp b/Engine/Plugins/Developer/Sandbox/SandboxedEditing/Source/SandboxedEditing/Private/Features/Browser/ViewModels/Workflows/LeaveWithDirtyPackagesWorkflow.cpp
deleted file mode 100644
diff --git a/Engine/Plugins/Developer/Sandbox/SandboxedEditing/Source/SandboxedEditing/Private/Features/Browser/ViewModels/Workflows/LeaveWithDirtyPackagesWorkflow.h b/Engine/Plugins/Developer/Sandbox/SandboxedEditing/Source/SandboxedEditing/Private/Features/Browser/ViewModels/Workflows/LeaveWithDirtyPackagesWorkflow.h
deleted file mode 100644
diff --git a/Engine/Plugins/Developer/Sandbox/SandboxedEditing/Source/SandboxedEditing/Private/Features/Browser/Views/AskUserAboutDirtyPackagesView.cpp b/Engine/Plugins/Developer/Sandbox/SandboxedEditing/Source/SandboxedEditing/Private/Features/Browser/Views/AskUserAboutDirtyPackagesView.cpp
new file mode 100644
--- /dev/null
+++ b/Engine/Plugins/Developer/Sandbox/SandboxedEditing/Source/SandboxedEditing/Private/Features/Browser/Views/AskUserAboutDirtyPackagesView.cpp
@@ -0,0 +1,97 @@
+﻿// Copyright Epic Games, Inc. All Rights Reserved.
+
+#include "AskUserAboutDirtyPackagesView.h"
+
+#include "Features/Browser/ViewModels/Leaving/LeaveSandboxViewModel.h"
+#include "Framework/Notifications/NotificationManager.h"
+#include "Widgets/Notifications/SNotificationList.h"
+
+#define LOCTEXT_NAMESPACE "FNotificationView"
+
+namespace UE::SandboxedEditing
+{
+FAskUserAboutDirtyPackagesView::FAskUserAboutDirtyPackagesView(const TSharedRef<FLeaveSandboxViewModel>& InLeaveViewModel)
+	: LeaveViewModel(InLeaveViewModel)
+{
+	LeaveViewModel->OnCreateLeaveWithDirtyPackageWorkflow().AddRaw(this, &FAskUserAboutDirtyPackagesView::ShowNotification);
+}
+
+FAskUserAboutDirtyPackagesView::~FAskUserAboutDirtyPackagesView()
+{
+	LeaveViewModel->OnCreateLeaveWithDirtyPackageWorkflow().RemoveAll(this);
+	
+	if (FAskUserAboutDirtyPackagesWorkflow* Workflow = LeaveViewModel->GetLeaveWithDirtyPackageWorkflow())
+	{
+		Workflow->OnWorkflowEnded().Remove(HandleWorkflowEndHandle);
+	}
+}
+
+void FAskUserAboutDirtyPackagesView::ShowNotification(FAskUserAboutDirtyPackagesWorkflow& InWorkflow)
+{
+	HideNotification();
+	HandleWorkflowEndHandle = InWorkflow.OnWorkflowEnded().AddLambda([this](auto){ HideNotification(); });
+	
+	FSlateNotificationManager& NotificationManager = FSlateNotificationManager::Get();
+	
+	FNotificationInfo NotificationInfo(LOCTEXT("Unsaved.Title", "In-memory changes"));
+	NotificationInfo.SubText = LOCTEXT("Unsaved.SubText", "You must either save or discard unsaved assets before leaving.");
+	NotificationInfo.bFireAndForget = false;
+	NotificationInfo.ButtonDetails = 
+	{
+		FNotificationButtonInfo(
+			LOCTEXT("Unsaved.SaveAll.Title", "Save all"),
+			LOCTEXT("Unsaved.SaveAll.ToolTip", "Save all unsaved assets"),
+			FSimpleDelegate::CreateRaw(this, &FAskUserAboutDirtyPackagesView::HandleSaveAll),
+			SNotificationItem::CS_None
+			),
+		FNotificationButtonInfo(
+			LOCTEXT("Unsaved.DiscardAll.Title", "Discard"),
+			LOCTEXT("Unsaved.DiscardAll.ToolTip", "Discard all unsaved assets"),
+			FSimpleDelegate::CreateRaw(this, &FAskUserAboutDirtyPackagesView::HandleRevert),
+			SNotificationItem::CS_None
+			),
+		FNotificationButtonInfo(
+			LOCTEXT("Unsaved.Cancel.Title", "Cancel"),
+			LOCTEXT("Unsaved.Cancel.ToolTip", "Does not leave the sandbox"),
+			FSimpleDelegate::CreateRaw(this, &FAskUserAboutDirtyPackagesView::HandleCancel),
+			SNotificationItem::CS_None
+			)
+	};
+	Notification = NotificationManager.AddNotification(NotificationInfo);
+}
+
+void FAskUserAboutDirtyPackagesView::HandleSaveAll() const
+{
+	if (FAskUserAboutDirtyPackagesWorkflow* Workflow = LeaveViewModel->GetLeaveWithDirtyPackageWorkflow())
+	{
+		Workflow->SaveDirtyPackages();
+	}
+}
+
+void FAskUserAboutDirtyPackagesView::HandleRevert() const
+{
+	if (FAskUserAboutDirtyPackagesWorkflow* Workflow = LeaveViewModel->GetLeaveWithDirtyPackageWorkflow())
+	{
+		Workflow->RevertDirtyPackages();
+	}
+}
+
+void FAskUserAboutDirtyPackagesView::HandleCancel() const
+{
+	if (FAskUserAboutDirtyPackagesWorkflow* Workflow = LeaveViewModel->GetLeaveWithDirtyPackageWorkflow())
+	{
+		Workflow->Cancel();
+	}
+}
+
+void FAskUserAboutDirtyPackagesView::HideNotification()
+{
+	if (Notification)
+	{
+		Notification->ExpireAndFadeout();
+		Notification.Reset();
+	}
+}
+}
+
+#undef LOCTEXT_NAMESPACE
\ No newline at end of file
diff --git a/Engine/Plugins/Developer/Sandbox/SandboxedEditing/Source/SandboxedEditing/Private/Features/Browser/Views/AskUserAboutDirtyPackagesView.h b/Engine/Plugins/Developer/Sandbox/SandboxedEditing/Source/SandboxedEditing/Private/Features/Browser/Views/AskUserAboutDirtyPackagesView.h
new file mode 100644
--- /dev/null
+++ b/Engine/Plugins/Developer/Sandbox/SandboxedEditing/Source/SandboxedEditing/Private/Features/Browser/Views/AskUserAboutDirtyPackagesView.h
@@ -0,0 +1,48 @@
+﻿// Copyright Epic Games, Inc. All Rights Reserved.
+
+#pragma once
+
+#include "Delegates/IDelegateInstance.h"
+#include "Templates/SharedPointer.h"
+#include "Templates/UnrealTemplate.h"
+
+class SNotificationItem;
+
+namespace UE::SandboxedEditing
+{
+class FLeaveSandboxViewModel;
+class FAskUserAboutDirtyPackagesWorkflow;
+
+/** Visualizes the notification for leaving the sandbox with dirty packages. */
+class FAskUserAboutDirtyPackagesView : public FNoncopyable
+{
+public:
+	
+	explicit FAskUserAboutDirtyPackagesView(const TSharedRef<FLeaveSandboxViewModel>& InLeaveViewModel);
+	~FAskUserAboutDirtyPackagesView();
+	
+private:
+	
+	/** Notifies us when a FLeaveSandboxViewModel starts. */
+	const TSharedRef<FLeaveSandboxViewModel> LeaveViewModel;
+
+	/** Whether currently showing the notification for when the user is trying to leave the sandbox even though there are modified assets. */
+	TSharedPtr<SNotificationItem> Notification;
+	
+	FDelegateHandle HandleWorkflowEndHandle;
+
+	/** Show the notification for when the user is trying to leave the sandbox even though there are modified assets. */
+	void ShowNotification(FAskUserAboutDirtyPackagesWorkflow& InWorkflow);
+	
+	/** Saves all dirty assets, hides the notification, and proceeds with the leave operation. */
+	void HandleSaveAll() const;
+	/** Discards all dirty assets, hides the notification, and proceeds with the leave operation. */
+	void HandleRevert() const;
+	/** Cancels the workflow. */
+	void HandleCancel() const;
+	
+	/** Hides the notification. */
+	void HideNotification();
+};
+}
+
