File Sandbox Core: Clean up the column system. Make it less generic because it was overgeneralized - making it too difficult to use.

#Synced-CL 49159615

diff --git a/Engine/Plugins/Developer/Sandbox/SandboxedEditing/Source/SandboxedEditing/Private/Features/Browser/BrowserFeature.cpp b/Engine/Plugins/Developer/Sandbox/SandboxedEditing/Source/SandboxedEditing/Private/Features/Browser/BrowserFeature.cpp
--- a/Engine/Plugins/Developer/Sandbox/SandboxedEditing/Source/SandboxedEditing/Private/Features/Browser/BrowserFeature.cpp
+++ b/Engine/Plugins/Developer/Sandbox/SandboxedEditing/Source/SandboxedEditing/Private/Features/Browser/BrowserFeature.cpp
@@ -2,15 +2,17 @@
 
 #include "BrowserFeature.h"
 
+#include "Framework/Application/SlateApplication.h"
 #include "Framework/Docking/TabManager.h"
 #include "SandboxedEditingStyle.h"
-#include "ViewModels/List/SandboxListViewModel.h"
+#include "ViewModels/SandboxControlsViewModel.h"
 #include "Widgets/Docking/SDockTab.h"
 #include "Widgets/SSandboxedEditingRoot.h"
+#include "Widgets/Unsandboxed/List/Columns/LastModifiedSandboxColumn.h"
+#include "Widgets/Unsandboxed/List/Columns/NameSandboxColumn.h"
+#include "Widgets/Unsandboxed/List/Columns/VersionSandboxColumn.h"
 #include "WorkspaceMenuStructure.h"
 #include "WorkspaceMenuStructureModule.h"
-#include "Framework/Application/SlateApplication.h"
-#include "ViewModels/SandboxControlsViewModel.h"
 
 #define LOCTEXT_NAMESPACE "FBrowserFeature"
 
@@ -20,24 +22,30 @@ namespace BrowserDetail
 {
 const FLazyName SandboxedEditingTabName(TEXT("SandboxedEditingTab"));
 
-static TMap<FName, TSharedRef<ISandboxColumnBehavior>> TransformToBehaviours(const TMap<FName, FSandboxColumnInfo>& InColumns)
+static void MakeSandboxColumns(TMap<FName, TSharedRef<ISandboxColumnBehavior>>& OutBehaviors, TMap<FName, TSharedRef<ISandboxColumnWidgetFactory>>& OutFactories)
 {
-	TMap<FName, TSharedRef<ISandboxColumnBehavior>> Result;
-	Algo::Transform(InColumns, Result, [](const TPair<FName, FSandboxColumnInfo>& Pair)
+	const TSharedRef<FNameSandboxColumn> Name = MakeShared<FNameSandboxColumn>();
+	const TSharedRef<FVersionSandboxColumn> Version = MakeShared<FVersionSandboxColumn>();
+	const TSharedRef<FLastModifiedSandboxColumn> LastModified = MakeShared<FLastModifiedSandboxColumn>();
+	OutBehaviors = TMap<FName, TSharedRef<ISandboxColumnBehavior>>
+	{
+		{ NameSandboxColumn, Name },
+		{ VersionSandboxNameColumn, Version },
+		{ LastModifiedSandboxColumn, LastModified }
+	};
+	OutFactories = TMap<FName, TSharedRef<ISandboxColumnWidgetFactory>>
 	{
-		return TPair<FName, TSharedRef<ISandboxColumnBehavior>>{ Pair.Key, Pair.Value.Behavior };
-	});
-	return Result;
+		{ NameSandboxColumn, Name },
+		{ VersionSandboxNameColumn, Version },
+		{ LastModifiedSandboxColumn, LastModified }
+	};
 }
 
-static TMap<FName, TSharedRef<ISandboxWidgetFactory>> TransformToWidgetFactories(const TMap<FName, FSandboxColumnInfo>& InColumns)
+static FSandboxColumnRegistry MakeSandboxColumnRegistry()
 {
-	TMap<FName, TSharedRef<ISandboxWidgetFactory>> Result;
-	Algo::Transform(InColumns, Result, [](const TPair<FName, FSandboxColumnInfo>& Pair)
-	{
-		return TPair<FName, TSharedRef<ISandboxWidgetFactory>>{ Pair.Key, Pair.Value.WidgetFactory };
-	});
-	return Result;
+	FSandboxColumnRegistry Registry;
+	MakeSandboxColumns(Registry.ColumnBehaviors, Registry.ColumnFactories);
+	return Registry;
 }
 }
 	
@@ -46,8 +54,8 @@ FBrowserFeature::FBrowserFeature(
 	const TSharedRef<FileSandboxUI::IExternalSandboxActiveViewModel>& InExternalSandboxViewModel
 	)
 	: SandboxModel(InSandboxModel)
-	, BrowserColumns(MakeBuiltInColumns())
-	, ViewModels(BrowserDetail::TransformToBehaviours(BrowserColumns), SandboxModel, InExternalSandboxViewModel)
+	, SandboxColumns(BrowserDetail::MakeSandboxColumnRegistry())
+	, ViewModels(SandboxColumns, SandboxModel, InExternalSandboxViewModel)
 	, CommandBindings(MakeShared<FBrowserCommandBindings>(ViewModels))
 	, LeaveWithDirtyPackagesView(ViewModels.LeaveViewModel)
 {
@@ -91,7 +99,7 @@ TSharedRef<SDockTab> FBrowserFeature::SpawnBrowserTab(const FSpawnTabArgs& InArg
 		[
 			SNew(SSandboxedEditingRoot, ViewModels)
 			.CommandList(CommandBindings->GetCommandList())
-			.ColumnFactories(BrowserDetail::TransformToWidgetFactories(BrowserColumns))
+			.ColumnFactories(SandboxColumns.ColumnFactories)
 		];
 }
 }
diff --git a/Engine/Plugins/Developer/Sandbox/SandboxedEditing/Source/SandboxedEditing/Private/Features/Browser/BrowserFeature.h b/Engine/Plugins/Developer/Sandbox/SandboxedEditing/Source/SandboxedEditing/Private/Features/Browser/BrowserFeature.h
--- a/Engine/Plugins/Developer/Sandbox/SandboxedEditing/Source/SandboxedEditing/Private/Features/Browser/BrowserFeature.h
+++ b/Engine/Plugins/Developer/Sandbox/SandboxedEditing/Source/SandboxedEditing/Private/Features/Browser/BrowserFeature.h
@@ -43,8 +43,8 @@ private:
 	/** The shared sandbox model. */
 	const TSharedRef<FSandboxSystemModel> SandboxModel;
 	
-	/** The columns used by the browser widget. */
-	const TMap<FName, FSandboxColumnInfo> BrowserColumns;
+	/** The columns that are displayed in the browser. */
+	const FSandboxColumnRegistry SandboxColumns;
 	/** Holds all view models relevant for the browser. */
 	const FBrowserViewModels ViewModels;
 	
diff --git a/Engine/Plugins/Developer/Sandbox/SandboxedEditing/Source/SandboxedEditing/Private/Features/Browser/ViewModels/BrowserViewModels.h b/Engine/Plugins/Developer/Sandbox/SandboxedEditing/Source/SandboxedEditing/Private/Features/Browser/ViewModels/BrowserViewModels.h
--- a/Engine/Plugins/Developer/Sandbox/SandboxedEditing/Source/SandboxedEditing/Private/Features/Browser/ViewModels/BrowserViewModels.h
+++ b/Engine/Plugins/Developer/Sandbox/SandboxedEditing/Source/SandboxedEditing/Private/Features/Browser/ViewModels/BrowserViewModels.h
@@ -9,6 +9,7 @@
 #include "List/SandboxListViewModel.h"
 #include "SandboxControlsViewModel.h"
 #include "SandboxMetaDataViewModel.h"
+#include "List/SandboxColumnRegistry.h"
 #include "Templates/SharedPointer.h"
 
 namespace UE::FileSandboxUI
@@ -18,12 +19,11 @@ class IExternalSandboxActiveViewModel;
 
 namespace UE::SandboxedEditing
 {
+class ISandboxColumnBehavior;
+
 /** Holds all view models used by the root UI. Make sharing view models easier. */
 struct FBrowserViewModels
 {
-	/** The columns used by the browser widget. */
-	const TMap<FName, TArray<FSandboxColumnInfo>> BrowserColumns;
-	
 	/** Knows how to filter sandbox items. */
 	const TSharedRef<FFilterSandboxViewModel> FilterSandboxViewModel;
 	/** Knows about the items that are displayed in the browser. */
@@ -44,17 +44,12 @@ struct FBrowserViewModels
 	const TSharedRef<FileSandboxUI::IExternalSandboxActiveViewModel> ExternalSandboxViewModel;
 	
 	explicit FBrowserViewModels(
-		TMap<FName, TSharedRef<ISandboxColumnBehavior>> InColumns, 
+		const FSandboxColumnRegistry& InSandboxColumns, 
 		const TSharedRef<FSandboxSystemModel>& InModel,
 		const TSharedRef<FileSandboxUI::IExternalSandboxActiveViewModel>& InExternalSandboxViewModel
 		)
-		: FilterSandboxViewModel(MakeShared<FFilterSandboxViewModel>([&InColumns]
-			{
-				TArray<TSharedRef<ISandboxColumnBehavior>> Columns;
-				InColumns.GenerateValueArray(Columns);
-				return Columns;
-			}()))
-		, ListViewModel(MakeShared<FSandboxListViewModel>(InModel, FilterSandboxViewModel, MoveTemp(InColumns)))
+		: FilterSandboxViewModel(MakeShared<FFilterSandboxViewModel>(InSandboxColumns.GetBehaviorArray()))
+		, ListViewModel(MakeShared<FSandboxListViewModel>(InModel, FilterSandboxViewModel, InSandboxColumns.ColumnBehaviors))
 		, ControlsViewModel(MakeShared<FSandboxControlsViewModel>(InModel))
 		, MetaDataViewModel(MakeShared<FSandboxMetaDataViewModel>(InModel))
 		, LeaveViewModel(MakeShared<FLeaveSandboxViewModel>(InModel))
diff --git a/Engine/Plugins/Developer/Sandbox/SandboxedEditing/Source/SandboxedEditing/Private/Features/Browser/ViewModels/Columns/ColumnInfo.h b/Engine/Plugins/Developer/Sandbox/SandboxedEditing/Source/SandboxedEditing/Private/Features/Browser/ViewModels/Columns/ColumnInfo.h
deleted file mode 100644
diff --git a/Engine/Plugins/Developer/Sandbox/SandboxedEditing/Source/SandboxedEditing/Private/Features/Browser/ViewModels/Columns/IColumnBehavior.h b/Engine/Plugins/Developer/Sandbox/SandboxedEditing/Source/SandboxedEditing/Private/Features/Browser/ViewModels/Columns/IColumnBehavior.h
deleted file mode 100644
diff --git a/Engine/Plugins/Developer/Sandbox/SandboxedEditing/Source/SandboxedEditing/Private/Features/Browser/ViewModels/Columns/IColumnWidgetFactory.h b/Engine/Plugins/Developer/Sandbox/SandboxedEditing/Source/SandboxedEditing/Private/Features/Browser/ViewModels/Columns/IColumnWidgetFactory.h
deleted file mode 100644
diff --git a/Engine/Plugins/Developer/Sandbox/SandboxedEditing/Source/SandboxedEditing/Private/Features/Browser/ViewModels/Columns/ISearchableColumnBehavior.h b/Engine/Plugins/Developer/Sandbox/SandboxedEditing/Source/SandboxedEditing/Private/Features/Browser/ViewModels/Columns/ISearchableColumnBehavior.h
new file mode 100644
--- /dev/null
+++ b/Engine/Plugins/Developer/Sandbox/SandboxedEditing/Source/SandboxedEditing/Private/Features/Browser/ViewModels/Columns/ISearchableColumnBehavior.h
@@ -0,0 +1,20 @@
+﻿// Copyright Epic Games, Inc. All Rights Reserved.
+
+#pragma once
+
+#include "Widgets/Views/SHeaderRow.h"
+
+namespace UE::SandboxedEditing
+{
+/** Interface for a column that can be searched */
+template<typename TRowData>
+class ISearchableColumnBehavior
+{
+public:
+
+	/** Search terms for this column */
+	virtual void PopulateSearchTerms(const TRowData& InRowData, TArray<FString>& OutSearchTerms) const = 0;
+
+	virtual ~ISearchableColumnBehavior() = default;
+};
+}
diff --git a/Engine/Plugins/Developer/Sandbox/SandboxedEditing/Source/SandboxedEditing/Private/Features/Browser/ViewModels/List/FilterSandboxViewModel.cpp b/Engine/Plugins/Developer/Sandbox/SandboxedEditing/Source/SandboxedEditing/Private/Features/Browser/ViewModels/List/FilterSandboxViewModel.cpp
--- a/Engine/Plugins/Developer/Sandbox/SandboxedEditing/Source/SandboxedEditing/Private/Features/Browser/ViewModels/List/FilterSandboxViewModel.cpp
+++ b/Engine/Plugins/Developer/Sandbox/SandboxedEditing/Source/SandboxedEditing/Private/Features/Browser/ViewModels/List/FilterSandboxViewModel.cpp
@@ -2,7 +2,7 @@
 
 #include "FilterSandboxViewModel.h"
 
-#include "Framework/Models/SandboxInfo.h"
+#include "ISandboxColumnBehavior.h"
 
 namespace UE::SandboxedEditing
 {
diff --git a/Engine/Plugins/Developer/Sandbox/SandboxedEditing/Source/SandboxedEditing/Private/Features/Browser/ViewModels/List/FilterSandboxViewModel.h b/Engine/Plugins/Developer/Sandbox/SandboxedEditing/Source/SandboxedEditing/Private/Features/Browser/ViewModels/List/FilterSandboxViewModel.h
--- a/Engine/Plugins/Developer/Sandbox/SandboxedEditing/Source/SandboxedEditing/Private/Features/Browser/ViewModels/List/FilterSandboxViewModel.h
+++ b/Engine/Plugins/Developer/Sandbox/SandboxedEditing/Source/SandboxedEditing/Private/Features/Browser/ViewModels/List/FilterSandboxViewModel.h
@@ -2,13 +2,14 @@
 
 #pragma once
 
-#include "Features/Browser/Widgets/Unsandboxed/List/Columns/SandboxColumns.h"
+#include "SandboxListItem.h"
 #include "Misc/Attribute.h"
 #include "Misc/TextFilter.h"
 #include "Templates/UnrealTemplate.h"
 
 namespace UE::SandboxedEditing
 {
+class ISandboxColumnBehavior;
 struct FSandboxInfo;
 
 /** Knows of all filters for filtering sandbox items. */
diff --git a/Engine/Plugins/Developer/Sandbox/SandboxedEditing/Source/SandboxedEditing/Private/Features/Browser/ViewModels/List/ISandboxColumnBehavior.h b/Engine/Plugins/Developer/Sandbox/SandboxedEditing/Source/SandboxedEditing/Private/Features/Browser/ViewModels/List/ISandboxColumnBehavior.h
new file mode 100644
--- /dev/null
+++ b/Engine/Plugins/Developer/Sandbox/SandboxedEditing/Source/SandboxedEditing/Private/Features/Browser/ViewModels/List/ISandboxColumnBehavior.h
@@ -0,0 +1,15 @@
+﻿// Copyright Epic Games, Inc. All Rights Reserved.
+
+#pragma once
+
+#include "Features/Browser/ViewModels/Columns/ISearchableColumnBehavior.h"
+#include "Templates/SharedPointer.h"
+
+namespace UE::SandboxedEditing
+{
+class FSandboxListItem;
+
+/** Operations for columns that display sandboxes in the sandbox list view. */
+class ISandboxColumnBehavior : public ISearchableColumnBehavior<TSharedPtr<FSandboxListItem>>
+{};
+}
diff --git a/Engine/Plugins/Developer/Sandbox/SandboxedEditing/Source/SandboxedEditing/Private/Features/Browser/ViewModels/List/ISandboxColumnWidgetFactory.h b/Engine/Plugins/Developer/Sandbox/SandboxedEditing/Source/SandboxedEditing/Private/Features/Browser/ViewModels/List/ISandboxColumnWidgetFactory.h
new file mode 100644
--- /dev/null
+++ b/Engine/Plugins/Developer/Sandbox/SandboxedEditing/Source/SandboxedEditing/Private/Features/Browser/ViewModels/List/ISandboxColumnWidgetFactory.h
@@ -0,0 +1,46 @@
+﻿// Copyright Epic Games, Inc. All Rights Reserved.
+
+#pragma once
+
+#include "Internationalization/Text.h"
+#include "Misc/Attribute.h"
+#include "UObject/NameTypes.h"
+#include "Widgets/Views/SHeaderRow.h"
+
+namespace UE::SandboxedEditing
+{
+class FSandboxListItem;
+
+/** Args for making a column widget in a row. */
+struct FMakeSandboxColumnWidgetArgs
+{
+	/** The row data for this column widget */
+	TSharedPtr<FSandboxListItem> RowData;
+		
+	/** Text that is being searched */
+	TAttribute<FText> HighlightText;
+	
+	/** Whether the row is currently selected in the SListView. Required for some widgets, e.g. SInlineEditableTextBlock. */
+	FIsSelected IsSelectedDelegate;
+
+	explicit FMakeSandboxColumnWidgetArgs(TSharedPtr<FSandboxListItem> InRowData, TAttribute<FText> InHighlightText, FIsSelected InIsSelectedDelegate)
+		: RowData(MoveTemp(InRowData))
+		, HighlightText(MoveTemp(InHighlightText))
+		, IsSelectedDelegate(MoveTemp(InIsSelectedDelegate))
+	{}
+};
+
+/** Interface for creating Slate widgets for a column. */
+class ISandboxColumnWidgetFactory
+{
+public:
+	
+	/** @return Info required to construct the column. */
+	virtual SHeaderRow::FColumn::FArguments MakeColumnArguments() = 0;
+	
+	/** @return Widget that should be put into the row for this column */
+	virtual TSharedRef<SWidget> MakeColumnWidget(const FMakeSandboxColumnWidgetArgs& InArgs) = 0;
+
+	virtual ~ISandboxColumnWidgetFactory() = default;
+};
+}
diff --git a/Engine/Plugins/Developer/Sandbox/SandboxedEditing/Source/SandboxedEditing/Private/Features/Browser/ViewModels/List/SandboxColumnRegistry.h b/Engine/Plugins/Developer/Sandbox/SandboxedEditing/Source/SandboxedEditing/Private/Features/Browser/ViewModels/List/SandboxColumnRegistry.h
new file mode 100644
--- /dev/null
+++ b/Engine/Plugins/Developer/Sandbox/SandboxedEditing/Source/SandboxedEditing/Private/Features/Browser/ViewModels/List/SandboxColumnRegistry.h
@@ -0,0 +1,28 @@
+// Copyright Epic Games, Inc. All Rights Reserved.
+
+#pragma once
+
+#include "Containers/Map.h"
+#include "SandboxColumns.h"
+#include "Templates/SharedPointer.h"
+
+namespace UE::SandboxedEditing
+{
+class ISandboxColumnBehavior;
+
+struct FSandboxColumnRegistry
+{
+	/** The behavior of the columns for view models. */
+	TMap<FName, TSharedRef<ISandboxColumnBehavior>> ColumnBehaviors;
+	
+	/** Creates the widgets for the columns. */
+	TMap<FName, TSharedRef<ISandboxColumnWidgetFactory>> ColumnFactories;
+	
+	TArray<TSharedRef<ISandboxColumnBehavior>> GetBehaviorArray() const
+	{
+		TArray<TSharedRef<ISandboxColumnBehavior>> Columns;
+		ColumnBehaviors.GenerateValueArray(Columns);
+		return Columns;
+	}
+};
+}
\ No newline at end of file
diff --git a/Engine/Plugins/Developer/Sandbox/SandboxedEditing/Source/SandboxedEditing/Private/Features/Browser/ViewModels/List/SandboxColumns.h b/Engine/Plugins/Developer/Sandbox/SandboxedEditing/Source/SandboxedEditing/Private/Features/Browser/ViewModels/List/SandboxColumns.h
new file mode 100644
--- /dev/null
+++ b/Engine/Plugins/Developer/Sandbox/SandboxedEditing/Source/SandboxedEditing/Private/Features/Browser/ViewModels/List/SandboxColumns.h
@@ -0,0 +1,20 @@
+﻿// Copyright Epic Games, Inc. All Rights Reserved.
+
+#pragma once
+
+#include "Templates/SharedPointer.h"
+#include "UObject/NameTypes.h"
+
+namespace UE::SandboxedEditing
+{
+class FSandboxListItem;
+class ISandboxColumnBehavior;
+class ISandboxColumnWidgetFactory;
+
+/** Needed for SLATE_ARGUMENT (the macro is written in such a way that TMap is not parsed correctly). */
+using FSandboxColumnFactoryMap = TMap<FName, TSharedRef<ISandboxColumnWidgetFactory>>;
+
+constexpr FLazyName NameSandboxColumn(TEXT("NameSandboxColumn"));
+constexpr FLazyName VersionSandboxNameColumn(TEXT("VersionSandboxNameColumn"));
+constexpr FLazyName LastModifiedSandboxColumn(TEXT("LastModifiedSandboxColumn"));
+}
diff --git a/Engine/Plugins/Developer/Sandbox/SandboxedEditing/Source/SandboxedEditing/Private/Features/Browser/ViewModels/List/SandboxListViewModel.h b/Engine/Plugins/Developer/Sandbox/SandboxedEditing/Source/SandboxedEditing/Private/Features/Browser/ViewModels/List/SandboxListViewModel.h
--- a/Engine/Plugins/Developer/Sandbox/SandboxedEditing/Source/SandboxedEditing/Private/Features/Browser/ViewModels/List/SandboxListViewModel.h
+++ b/Engine/Plugins/Developer/Sandbox/SandboxedEditing/Source/SandboxedEditing/Private/Features/Browser/ViewModels/List/SandboxListViewModel.h
@@ -3,8 +3,6 @@
 #pragma once
 
 #include "Containers/Array.h"
-#include "Features/Browser/ViewModels/Columns/IColumnBehavior.h"
-#include "Features/Browser/Widgets/Unsandboxed/List/Columns/SandboxColumns.h"
 #include "Framework/Models/SandboxSystemModel.h"
 #include "Templates/UnrealTemplate.h"
 #include "Widgets/Views/SHeaderRow.h"
@@ -13,6 +11,7 @@ namespace UE::SandboxedEditing
 {
 class FFilterSandboxViewModel;
 class FSandboxListItem;
+class ISandboxColumnBehavior;
 
 /** Decides what is displayed in SSandboxListView. */
 class FSandboxListViewModel : public FNoncopyable
diff --git a/Engine/Plugins/Developer/Sandbox/SandboxedEditing/Source/SandboxedEditing/Private/Features/Browser/Widgets/SSandboxedEditingRoot.h b/Engine/Plugins/Developer/Sandbox/SandboxedEditing/Source/SandboxedEditing/Private/Features/Browser/Widgets/SSandboxedEditingRoot.h
--- a/Engine/Plugins/Developer/Sandbox/SandboxedEditing/Source/SandboxedEditing/Private/Features/Browser/Widgets/SSandboxedEditingRoot.h
+++ b/Engine/Plugins/Developer/Sandbox/SandboxedEditing/Source/SandboxedEditing/Private/Features/Browser/Widgets/SSandboxedEditingRoot.h
@@ -2,9 +2,8 @@
 
 #pragma once
 
-#include "Features/Browser/ViewModels/Active/ActiveSandboxDetailsViewModel.h"
+#include "Features/Browser/ViewModels/List/SandboxColumns.h"
 #include "Sandboxed/SActiveSandbox.h"
-#include "Unsandboxed/List/Columns/SandboxColumns.h"
 #include "Widgets/DeclarativeSyntaxSupport.h"
 #include "Widgets/SCompoundWidget.h"
 
diff --git a/Engine/Plugins/Developer/Sandbox/SandboxedEditing/Source/SandboxedEditing/Private/Features/Browser/Widgets/Unsandboxed/List/Columns/LastModifiedSandboxColumn.cpp b/Engine/Plugins/Developer/Sandbox/SandboxedEditing/Source/SandboxedEditing/Private/Features/Browser/Widgets/Unsandboxed/List/Columns/LastModifiedSandboxColumn.cpp
--- a/Engine/Plugins/Developer/Sandbox/SandboxedEditing/Source/SandboxedEditing/Private/Features/Browser/Widgets/Unsandboxed/List/Columns/LastModifiedSandboxColumn.cpp
+++ b/Engine/Plugins/Developer/Sandbox/SandboxedEditing/Source/SandboxedEditing/Private/Features/Browser/Widgets/Unsandboxed/List/Columns/LastModifiedSandboxColumn.cpp
@@ -2,6 +2,7 @@
 
 #include "LastModifiedSandboxColumn.h"
 
+#include "Features/Browser/ViewModels/List/SandboxColumns.h"
 #include "Features/Browser/ViewModels/List/SandboxListItem.h"
 #include "Framework/Models/SandboxInfo.h"
 #include "SandboxedEditingStyle.h"
@@ -20,7 +21,7 @@ static FText FormatLastModified(const FSandboxInfo& InRowData)
 	return FText::AsDateTime(InRowData.LastModified);
 }
 }
-void FLastModifiedSandboxColumn::PopulateSearchTerms(const TSharedRef<FSandboxListItem>& InRowData, TArray<FString>& OutSearchTerms) const
+void FLastModifiedSandboxColumn::PopulateSearchTerms(const TSharedPtr<FSandboxListItem>& InRowData, TArray<FString>& OutSearchTerms) const
 {
 	OutSearchTerms.Add(
 		LastModifiedDetail::FormatLastModified(InRowData->SandboxInfo).ToString()
@@ -36,7 +37,7 @@ SHeaderRow::FColumn::FArguments FLastModifiedSandboxColumn::MakeColumnArguments(
 		.ToolTipText(LOCTEXT("LastModified.Description", "When this sandbox was last modified"));
 }
 
-TSharedRef<SWidget> FLastModifiedSandboxColumn::MakeColumnWidget(const FMakeColumnWidgetArgs& InArgs)
+TSharedRef<SWidget> FLastModifiedSandboxColumn::MakeColumnWidget(const FMakeSandboxColumnWidgetArgs& InArgs)
 {
 	return SNew(SBox)
 		.VAlign(VAlign_Center)
diff --git a/Engine/Plugins/Developer/Sandbox/SandboxedEditing/Source/SandboxedEditing/Private/Features/Browser/Widgets/Unsandboxed/List/Columns/LastModifiedSandboxColumn.h b/Engine/Plugins/Developer/Sandbox/SandboxedEditing/Source/SandboxedEditing/Private/Features/Browser/Widgets/Unsandboxed/List/Columns/LastModifiedSandboxColumn.h
--- a/Engine/Plugins/Developer/Sandbox/SandboxedEditing/Source/SandboxedEditing/Private/Features/Browser/Widgets/Unsandboxed/List/Columns/LastModifiedSandboxColumn.h
+++ b/Engine/Plugins/Developer/Sandbox/SandboxedEditing/Source/SandboxedEditing/Private/Features/Browser/Widgets/Unsandboxed/List/Columns/LastModifiedSandboxColumn.h
@@ -2,22 +2,23 @@
 
 #pragma once
 
-#include "SandboxColumns.h"
+#include "Features/Browser/ViewModels/List/ISandboxColumnWidgetFactory.h"
+#include "Features/Browser/ViewModels/List/ISandboxColumnBehavior.h"
 
 namespace UE::SandboxedEditing
 {
 /** Displays when the sandbox was last modified. */
-class FLastModifiedSandboxColumn : public ISandboxColumnBehavior, public ISandboxWidgetFactory
+class FLastModifiedSandboxColumn : public ISandboxColumnBehavior, public ISandboxColumnWidgetFactory
 {
 public:
 	
 	//~ Begin ISandboxColumnBehavior Interface
-	virtual void PopulateSearchTerms(const TSharedRef<FSandboxListItem>& InRowData, TArray<FString>& OutSearchTerms) const override;
+	virtual void PopulateSearchTerms(const TSharedPtr<FSandboxListItem>& InRowData, TArray<FString>& OutSearchTerms) const override;
 	//~ End ISandboxColumnBehavior Interface
 	
 	//~ Begin ISandboxWidgetFactory Interface
 	virtual SHeaderRow::FColumn::FArguments MakeColumnArguments() override;
-	virtual TSharedRef<SWidget> MakeColumnWidget(const FMakeColumnWidgetArgs& InArgs) override;
+	virtual TSharedRef<SWidget> MakeColumnWidget(const FMakeSandboxColumnWidgetArgs& InArgs) override;
 	//~ End ISandboxWidgetFactory Interface
 };
 }
diff --git a/Engine/Plugins/Developer/Sandbox/SandboxedEditing/Source/SandboxedEditing/Private/Features/Browser/Widgets/Unsandboxed/List/Columns/NameSandboxColumn.cpp b/Engine/Plugins/Developer/Sandbox/SandboxedEditing/Source/SandboxedEditing/Private/Features/Browser/Widgets/Unsandboxed/List/Columns/NameSandboxColumn.cpp
--- a/Engine/Plugins/Developer/Sandbox/SandboxedEditing/Source/SandboxedEditing/Private/Features/Browser/Widgets/Unsandboxed/List/Columns/NameSandboxColumn.cpp
+++ b/Engine/Plugins/Developer/Sandbox/SandboxedEditing/Source/SandboxedEditing/Private/Features/Browser/Widgets/Unsandboxed/List/Columns/NameSandboxColumn.cpp
@@ -2,6 +2,7 @@
 
 #include "NameSandboxColumn.h"
 
+#include "Features/Browser/ViewModels/List/SandboxColumns.h"
 #include "Features/Browser/ViewModels/List/SandboxListItem.h"
 #include "Framework/Models/SandboxInfo.h"
 #include "SandboxedEditingStyle.h"
@@ -13,7 +14,7 @@
 
 namespace UE::SandboxedEditing
 {
-void FNameSandboxColumn::PopulateSearchTerms(const TSharedRef<FSandboxListItem>& InRowData, TArray<FString>& OutSearchTerms) const
+void FNameSandboxColumn::PopulateSearchTerms(const TSharedPtr<FSandboxListItem>& InRowData, TArray<FString>& OutSearchTerms) const
 {
 	OutSearchTerms.Add(InRowData->SandboxInfo.Name);
 }
@@ -28,7 +29,7 @@ SHeaderRow::FColumn::FArguments FNameSandboxColumn::MakeColumnArguments()
 		.ShouldGenerateWidget(true);
 }
 
-TSharedRef<SWidget> FNameSandboxColumn::MakeColumnWidget(const FMakeColumnWidgetArgs& InArgs)
+TSharedRef<SWidget> FNameSandboxColumn::MakeColumnWidget(const FMakeSandboxColumnWidgetArgs& InArgs)
 {
 	return SNew(SBox)
 		.VAlign(VAlign_Center)
diff --git a/Engine/Plugins/Developer/Sandbox/SandboxedEditing/Source/SandboxedEditing/Private/Features/Browser/Widgets/Unsandboxed/List/Columns/NameSandboxColumn.h b/Engine/Plugins/Developer/Sandbox/SandboxedEditing/Source/SandboxedEditing/Private/Features/Browser/Widgets/Unsandboxed/List/Columns/NameSandboxColumn.h
--- a/Engine/Plugins/Developer/Sandbox/SandboxedEditing/Source/SandboxedEditing/Private/Features/Browser/Widgets/Unsandboxed/List/Columns/NameSandboxColumn.h
+++ b/Engine/Plugins/Developer/Sandbox/SandboxedEditing/Source/SandboxedEditing/Private/Features/Browser/Widgets/Unsandboxed/List/Columns/NameSandboxColumn.h
@@ -2,22 +2,23 @@
 
 #pragma once
 
-#include "SandboxColumns.h"
+#include "Features/Browser/ViewModels/List/ISandboxColumnBehavior.h"
+#include "Features/Browser/ViewModels/List/ISandboxColumnWidgetFactory.h"
 
 namespace UE::SandboxedEditing
 {
 /** Displays the sandbox name. */
-class FNameSandboxColumn : public ISandboxColumnBehavior, public ISandboxWidgetFactory
+class FNameSandboxColumn : public ISandboxColumnBehavior, public ISandboxColumnWidgetFactory
 {
 public:
 	
 	//~ Begin ISandboxColumnBehavior Interface
-	virtual void PopulateSearchTerms(const TSharedRef<FSandboxListItem>& InRowData, TArray<FString>& OutSearchTerms) const override;
+	virtual void PopulateSearchTerms(const TSharedPtr<FSandboxListItem>& InRowData, TArray<FString>& OutSearchTerms) const override;
 	//~ End ISandboxColumnBehavior Interface
 	
 	//~ Begin ISandboxWidgetFactory Interface
 	virtual SHeaderRow::FColumn::FArguments MakeColumnArguments() override;
-	virtual TSharedRef<SWidget> MakeColumnWidget(const FMakeColumnWidgetArgs& InArgs) override;
+	virtual TSharedRef<SWidget> MakeColumnWidget(const FMakeSandboxColumnWidgetArgs& InArgs) override;
 	//~ End ISandboxWidgetFactory Interface
 };
 }
diff --git a/Engine/Plugins/Developer/Sandbox/SandboxedEditing/Source/SandboxedEditing/Private/Features/Browser/Widgets/Unsandboxed/List/Columns/SandboxColumns.cpp b/Engine/Plugins/Developer/Sandbox/SandboxedEditing/Source/SandboxedEditing/Private/Features/Browser/Widgets/Unsandboxed/List/Columns/SandboxColumns.cpp
deleted file mode 100644
diff --git a/Engine/Plugins/Developer/Sandbox/SandboxedEditing/Source/SandboxedEditing/Private/Features/Browser/Widgets/Unsandboxed/List/Columns/SandboxColumns.h b/Engine/Plugins/Developer/Sandbox/SandboxedEditing/Source/SandboxedEditing/Private/Features/Browser/Widgets/Unsandboxed/List/Columns/SandboxColumns.h
deleted file mode 100644
diff --git a/Engine/Plugins/Developer/Sandbox/SandboxedEditing/Source/SandboxedEditing/Private/Features/Browser/Widgets/Unsandboxed/List/Columns/VersionSandboxColumn.cpp b/Engine/Plugins/Developer/Sandbox/SandboxedEditing/Source/SandboxedEditing/Private/Features/Browser/Widgets/Unsandboxed/List/Columns/VersionSandboxColumn.cpp
--- a/Engine/Plugins/Developer/Sandbox/SandboxedEditing/Source/SandboxedEditing/Private/Features/Browser/Widgets/Unsandboxed/List/Columns/VersionSandboxColumn.cpp
+++ b/Engine/Plugins/Developer/Sandbox/SandboxedEditing/Source/SandboxedEditing/Private/Features/Browser/Widgets/Unsandboxed/List/Columns/VersionSandboxColumn.cpp
@@ -2,6 +2,7 @@
 
 #include "VersionSandboxColumn.h"
 
+#include "Features/Browser/ViewModels/List/SandboxColumns.h"
 #include "Features/Browser/ViewModels/List/SandboxListItem.h"
 #include "Framework/Models/SandboxInfo.h"
 #include "SandboxedEditingStyle.h"
@@ -20,7 +21,7 @@ static FText FormatVersion(const FSandboxInfo& InRowData)
 	return FText::AsCultureInvariant(InRowData.EngineVersion.ToString());
 }
 }
-void FVersionSandboxColumn::PopulateSearchTerms(const TSharedRef<FSandboxListItem>& InRowData, TArray<FString>& OutSearchTerms) const
+void FVersionSandboxColumn::PopulateSearchTerms(const TSharedPtr<FSandboxListItem>& InRowData, TArray<FString>& OutSearchTerms) const
 {
 	OutSearchTerms.Add(
 		VersionDetail::FormatVersion(InRowData->SandboxInfo).ToString()
@@ -36,7 +37,7 @@ SHeaderRow::FColumn::FArguments FVersionSandboxColumn::MakeColumnArguments()
 		.ToolTipText(LOCTEXT("Version.Description", "The engine version this sandbox had when it was created."));
 }
 
-TSharedRef<SWidget> FVersionSandboxColumn::MakeColumnWidget(const FMakeColumnWidgetArgs& InArgs)
+TSharedRef<SWidget> FVersionSandboxColumn::MakeColumnWidget(const FMakeSandboxColumnWidgetArgs& InArgs)
 {
 	return SNew(SBox)
 		.VAlign(VAlign_Center)
diff --git a/Engine/Plugins/Developer/Sandbox/SandboxedEditing/Source/SandboxedEditing/Private/Features/Browser/Widgets/Unsandboxed/List/Columns/VersionSandboxColumn.h b/Engine/Plugins/Developer/Sandbox/SandboxedEditing/Source/SandboxedEditing/Private/Features/Browser/Widgets/Unsandboxed/List/Columns/VersionSandboxColumn.h
--- a/Engine/Plugins/Developer/Sandbox/SandboxedEditing/Source/SandboxedEditing/Private/Features/Browser/Widgets/Unsandboxed/List/Columns/VersionSandboxColumn.h
+++ b/Engine/Plugins/Developer/Sandbox/SandboxedEditing/Source/SandboxedEditing/Private/Features/Browser/Widgets/Unsandboxed/List/Columns/VersionSandboxColumn.h
@@ -2,22 +2,23 @@
 
 #pragma once
 
-#include "SandboxColumns.h"
+#include "Features/Browser/ViewModels/List/ISandboxColumnBehavior.h"
+#include "Features/Browser/ViewModels/List/ISandboxColumnWidgetFactory.h"
 
 namespace UE::SandboxedEditing
 {
 /** Displays the engine version this sandbox was last used in. */
-class FVersionSandboxColumn : public ISandboxColumnBehavior, public ISandboxWidgetFactory
+class FVersionSandboxColumn : public ISandboxColumnBehavior, public ISandboxColumnWidgetFactory
 {
 public:
 	
 	//~ Begin ISandboxColumnBehavior Interface
-	virtual void PopulateSearchTerms(const TSharedRef<FSandboxListItem>& InRowData, TArray<FString>& OutSearchTerms) const override;
+	virtual void PopulateSearchTerms(const TSharedPtr<FSandboxListItem>& InRowData, TArray<FString>& OutSearchTerms) const override;
 	//~ End ISandboxColumnBehavior Interface
 	
 	//~ Begin ISandboxWidgetFactory Interface
 	virtual SHeaderRow::FColumn::FArguments MakeColumnArguments() override;
-	virtual TSharedRef<SWidget> MakeColumnWidget(const FMakeColumnWidgetArgs& InArgs) override;
+	virtual TSharedRef<SWidget> MakeColumnWidget(const FMakeSandboxColumnWidgetArgs& InArgs) override;
 	//~ End ISandboxWidgetFactory Interface
 };
 }
diff --git a/Engine/Plugins/Developer/Sandbox/SandboxedEditing/Source/SandboxedEditing/Private/Features/Browser/Widgets/Unsandboxed/List/SCreateSandboxListRow.cpp b/Engine/Plugins/Developer/Sandbox/SandboxedEditing/Source/SandboxedEditing/Private/Features/Browser/Widgets/Unsandboxed/List/SCreateSandboxListRow.cpp
--- a/Engine/Plugins/Developer/Sandbox/SandboxedEditing/Source/SandboxedEditing/Private/Features/Browser/Widgets/Unsandboxed/List/SCreateSandboxListRow.cpp
+++ b/Engine/Plugins/Developer/Sandbox/SandboxedEditing/Source/SandboxedEditing/Private/Features/Browser/Widgets/Unsandboxed/List/SCreateSandboxListRow.cpp
@@ -2,7 +2,7 @@
 
 #include "SCreateSandboxListRow.h"
 
-#include "Columns/SandboxColumns.h"
+#include "Features/Browser/ViewModels/List/SandboxColumns.h"
 #include "Features/Browser/ViewModels/SandboxCreationWorkflow.h"
 #include "Framework/Application/SlateApplication.h"
 #include "Utils/WidgetUtils.h"
diff --git a/Engine/Plugins/Developer/Sandbox/SandboxedEditing/Source/SandboxedEditing/Private/Features/Browser/Widgets/Unsandboxed/List/SSandboxListRow.cpp b/Engine/Plugins/Developer/Sandbox/SandboxedEditing/Source/SandboxedEditing/Private/Features/Browser/Widgets/Unsandboxed/List/SSandboxListRow.cpp
--- a/Engine/Plugins/Developer/Sandbox/SandboxedEditing/Source/SandboxedEditing/Private/Features/Browser/Widgets/Unsandboxed/List/SSandboxListRow.cpp
+++ b/Engine/Plugins/Developer/Sandbox/SandboxedEditing/Source/SandboxedEditing/Private/Features/Browser/Widgets/Unsandboxed/List/SSandboxListRow.cpp
@@ -2,6 +2,7 @@
 
 #include "SSandboxListRow.h"
 
+#include "Features/Browser/ViewModels/List/ISandboxColumnWidgetFactory.h"
 #include "Features/Browser/ViewModels/List/SandboxListItem.h"
 
 namespace UE::SandboxedEditing
@@ -22,10 +23,10 @@ void SSandboxListRow::Construct(const FArguments& InArgs, const TSharedRef<FSand
 
 TSharedRef<SWidget> SSandboxListRow::GenerateWidgetForColumn(const FName& InColumnName)
 {
-	if (TSharedRef<ISandboxWidgetFactory>* Factory = ColumnFactoryMap.Find(InColumnName))
+	if (TSharedRef<ISandboxColumnWidgetFactory>* Factory = ColumnFactoryMap.Find(InColumnName))
 	{
 		return Factory->Get().MakeColumnWidget(
-			ISandboxWidgetFactory::FMakeColumnWidgetArgs(
+			FMakeSandboxColumnWidgetArgs(
 				Item.ToSharedRef(), HighlightTextAttr,
 				FIsSelected::CreateLambda([this]{ return IsSelectedAttr.Get(); })
 			)
diff --git a/Engine/Plugins/Developer/Sandbox/SandboxedEditing/Source/SandboxedEditing/Private/Features/Browser/Widgets/Unsandboxed/List/SSandboxListRow.h b/Engine/Plugins/Developer/Sandbox/SandboxedEditing/Source/SandboxedEditing/Private/Features/Browser/Widgets/Unsandboxed/List/SSandboxListRow.h
--- a/Engine/Plugins/Developer/Sandbox/SandboxedEditing/Source/SandboxedEditing/Private/Features/Browser/Widgets/Unsandboxed/List/SSandboxListRow.h
+++ b/Engine/Plugins/Developer/Sandbox/SandboxedEditing/Source/SandboxedEditing/Private/Features/Browser/Widgets/Unsandboxed/List/SSandboxListRow.h
@@ -2,7 +2,8 @@
 
 #pragma once
 
-#include "Columns/SandboxColumns.h"
+
+#include "Features/Browser/ViewModels/List/SandboxColumns.h"
 #include "Widgets/DeclarativeSyntaxSupport.h"
 #include "Widgets/Views/SListView.h"
 #include "Widgets/Views/STableRow.h"
diff --git a/Engine/Plugins/Developer/Sandbox/SandboxedEditing/Source/SandboxedEditing/Private/Features/Browser/Widgets/Unsandboxed/List/SSandboxListView.cpp b/Engine/Plugins/Developer/Sandbox/SandboxedEditing/Source/SandboxedEditing/Private/Features/Browser/Widgets/Unsandboxed/List/SSandboxListView.cpp
--- a/Engine/Plugins/Developer/Sandbox/SandboxedEditing/Source/SandboxedEditing/Private/Features/Browser/Widgets/Unsandboxed/List/SSandboxListView.cpp
+++ b/Engine/Plugins/Developer/Sandbox/SandboxedEditing/Source/SandboxedEditing/Private/Features/Browser/Widgets/Unsandboxed/List/SSandboxListView.cpp
@@ -2,6 +2,7 @@
 
 #include "SSandboxListView.h"
 
+#include "Features/Browser/ViewModels/List/ISandboxColumnWidgetFactory.h"
 #include "Features/Browser/ViewModels/List/SandboxListItem.h"
 #include "Features/Browser/ViewModels/List/SandboxListViewModel.h"
 #include "Features/Browser/ViewModels/SandboxControlsViewModel.h"
@@ -9,7 +10,6 @@
 #include "SSandboxListRow.h"
 #include "SandboxedEditingStyle.h"
 #include "Utils/BreakBehavior.h"
-#include "Widgets/Text/STextBlock.h"
 #include "Widgets/Views/SListView.h"
 
 #define LOCTEXT_NAMESPACE "SSandboxListView"
@@ -78,7 +78,7 @@ TSharedRef<SHeaderRow> SSandboxListView::MakeHeaderRow()
 		
 	for (FName ColumnId : ListViewModel->GetDisplayedColumns())
 	{
-		if (const TSharedRef<ISandboxWidgetFactory>* Factory = ColumnFactoryMap.Find(ColumnId))
+		if (const TSharedRef<ISandboxColumnWidgetFactory>* Factory = ColumnFactoryMap.Find(ColumnId))
 		{
 			HeaderRow->AddColumn(Factory->Get().MakeColumnArguments());
 		}
diff --git a/Engine/Plugins/Developer/Sandbox/SandboxedEditing/Source/SandboxedEditing/Private/Features/Browser/Widgets/Unsandboxed/List/SSandboxListView.h b/Engine/Plugins/Developer/Sandbox/SandboxedEditing/Source/SandboxedEditing/Private/Features/Browser/Widgets/Unsandboxed/List/SSandboxListView.h
--- a/Engine/Plugins/Developer/Sandbox/SandboxedEditing/Source/SandboxedEditing/Private/Features/Browser/Widgets/Unsandboxed/List/SSandboxListView.h
+++ b/Engine/Plugins/Developer/Sandbox/SandboxedEditing/Source/SandboxedEditing/Private/Features/Browser/Widgets/Unsandboxed/List/SSandboxListView.h
@@ -3,8 +3,8 @@
 #pragma once
 
 #include "DummySandboxItem.h"
+#include "Features/Browser/ViewModels/List/SandboxColumns.h"
 #include "Features/Browser/ViewModels/SandboxCreationWorkflow.h"
-#include "Columns/SandboxColumns.h"
 #include "Widgets/DeclarativeSyntaxSupport.h"
 #include "Widgets/SCompoundWidget.h"
 
diff --git a/Engine/Plugins/Developer/Sandbox/SandboxedEditing/Source/SandboxedEditing/Private/Features/Browser/Widgets/Unsandboxed/SSandboxedEditingBrowser.h b/Engine/Plugins/Developer/Sandbox/SandboxedEditing/Source/SandboxedEditing/Private/Features/Browser/Widgets/Unsandboxed/SSandboxedEditingBrowser.h
--- a/Engine/Plugins/Developer/Sandbox/SandboxedEditing/Source/SandboxedEditing/Private/Features/Browser/Widgets/Unsandboxed/SSandboxedEditingBrowser.h
+++ b/Engine/Plugins/Developer/Sandbox/SandboxedEditing/Source/SandboxedEditing/Private/Features/Browser/Widgets/Unsandboxed/SSandboxedEditingBrowser.h
@@ -2,7 +2,7 @@
 
 #pragma once
 
-#include "List/Columns/SandboxColumns.h"
+#include "Features/Browser/ViewModels/List/SandboxColumns.h"
 #include "Widgets/DeclarativeSyntaxSupport.h"
 #include "Widgets/SCompoundWidget.h"
 
