Sandboxed Editing: Implement workflow for persisting files on leave. This is similar to the flow in Multi User.
- A dialogue box is show whether you want to persist any files
- If yes is selected, a window is shown in which you can select the file to persist.
- This also adds a generic widget for displaying file states (add, remove, edit)

#Synced-CL 49159615

diff --git a/Engine/Plugins/Developer/Sandbox/FileSandbox/Source/FileSandboxCore/Private/Sandbox/ISandboxInstance.cpp b/Engine/Plugins/Developer/Sandbox/FileSandbox/Source/FileSandboxCore/Private/Sandbox/ISandboxInstance.cpp
--- a/Engine/Plugins/Developer/Sandbox/FileSandbox/Source/FileSandboxCore/Private/Sandbox/ISandboxInstance.cpp
+++ b/Engine/Plugins/Developer/Sandbox/FileSandbox/Source/FileSandboxCore/Private/Sandbox/ISandboxInstance.cpp
@@ -38,4 +38,15 @@ FGatheredFileChanges ISandboxInstance::GatherChangedFiles(EFileChangeGatherFlags
 	});
 	return Result;
 }
+
+bool ISandboxInstance::HasFileChanges() const
+{
+	bool bHasChanges = false;
+	EnumerateFileChanges([&bHasChanges](const FString& InFile, ESandboxFileChange InChange)
+	{
+		bHasChanges = true;
+		return EBreakBehavior::Break;
+	});
+	return bHasChanges;
+}
 }
diff --git a/Engine/Plugins/Developer/Sandbox/FileSandbox/Source/FileSandboxCore/Public/ISandboxInstance.h b/Engine/Plugins/Developer/Sandbox/FileSandbox/Source/FileSandboxCore/Public/ISandboxInstance.h
--- a/Engine/Plugins/Developer/Sandbox/FileSandbox/Source/FileSandboxCore/Public/ISandboxInstance.h
+++ b/Engine/Plugins/Developer/Sandbox/FileSandbox/Source/FileSandboxCore/Public/ISandboxInstance.h
@@ -65,6 +65,13 @@ public:
 	 * @see ISandboxManager::EnumerateFileChanges.
 	 */
 	virtual void EnumerateFileChanges(TFunctionRef<FProcessFileChangeSignature> InProcess) const = 0;
+	
+	/** 
+	 * @return Whether any changes have to any files.
+	 * @note This does not include packages with in-memory changes only.
+	 * @see GetDirtyPackages
+	 */
+	UE_API bool HasFileChanges() const;
 
 	
 	/** @return Whether the given package file exists on non sandbox path, i.e. whether the file was deleted in sandbox mode. */
diff --git a/Engine/Plugins/Developer/Sandbox/SandboxedEditing/Source/SandboxedEditing/Private/Features/Browser/BrowserFeature.cpp b/Engine/Plugins/Developer/Sandbox/SandboxedEditing/Source/SandboxedEditing/Private/Features/Browser/BrowserFeature.cpp
--- a/Engine/Plugins/Developer/Sandbox/SandboxedEditing/Source/SandboxedEditing/Private/Features/Browser/BrowserFeature.cpp
+++ b/Engine/Plugins/Developer/Sandbox/SandboxedEditing/Source/SandboxedEditing/Private/Features/Browser/BrowserFeature.cpp
@@ -56,8 +56,8 @@ FBrowserFeature::FBrowserFeature(
 	: SandboxModel(InSandboxModel)
 	, SandboxColumns(BrowserDetail::MakeSandboxColumnRegistry())
 	, ViewModels(SandboxColumns, SandboxModel, InExternalSandboxViewModel)
+	, Views(ViewModels.PersistViewModel, ViewModels.LeaveViewModel)
 	, CommandBindings(MakeShared<FBrowserCommandBindings>(ViewModels))
-	, LeaveWithDirtyPackagesView(ViewModels.LeaveViewModel)
 {
 	RegisterTabSpawner();
 }
diff --git a/Engine/Plugins/Developer/Sandbox/SandboxedEditing/Source/SandboxedEditing/Private/Features/Browser/BrowserFeature.h b/Engine/Plugins/Developer/Sandbox/SandboxedEditing/Source/SandboxedEditing/Private/Features/Browser/BrowserFeature.h
--- a/Engine/Plugins/Developer/Sandbox/SandboxedEditing/Source/SandboxedEditing/Private/Features/Browser/BrowserFeature.h
+++ b/Engine/Plugins/Developer/Sandbox/SandboxedEditing/Source/SandboxedEditing/Private/Features/Browser/BrowserFeature.h
@@ -7,6 +7,7 @@
 #include "Templates/UnrealTemplate.h"
 #include "ViewModels/BrowserViewModels.h"
 #include "Views/AskUserAboutDirtyPackagesView.h"
+#include "Views/BrowserViews.h"
 
 class SDockTab;
 class FSpawnTabArgs;
@@ -47,13 +48,12 @@ private:
 	const FSandboxColumnRegistry SandboxColumns;
 	/** Holds all view models relevant for the browser. */
 	const FBrowserViewModels ViewModels;
+	/** Holds all explicit views. */
+	const FBrowserViews Views;
 	
 	/** Binds user commands to models.*/
 	const TSharedRef<FBrowserCommandBindings> CommandBindings;
 	
-	/** Displays editor notifications (toaster, dialogue, etc.) for when the user tries to leave the sandbox with in-memory changes. */
-	FAskUserAboutDirtyPackagesView LeaveWithDirtyPackagesView;
-	
 	void RegisterTabSpawner();
 	void UnregisterTabSpawner();
 	TSharedRef<SDockTab> SpawnBrowserTab(const FSpawnTabArgs& InArgs);
diff --git a/Engine/Plugins/Developer/Sandbox/SandboxedEditing/Source/SandboxedEditing/Private/Features/Browser/ViewModels/BrowserViewModels.h b/Engine/Plugins/Developer/Sandbox/SandboxedEditing/Source/SandboxedEditing/Private/Features/Browser/ViewModels/BrowserViewModels.h
--- a/Engine/Plugins/Developer/Sandbox/SandboxedEditing/Source/SandboxedEditing/Private/Features/Browser/ViewModels/BrowserViewModels.h
+++ b/Engine/Plugins/Developer/Sandbox/SandboxedEditing/Source/SandboxedEditing/Private/Features/Browser/ViewModels/BrowserViewModels.h
@@ -7,15 +7,13 @@
 #include "Leaving/LeaveSandboxViewModel.h"
 #include "List/FilterSandboxViewModel.h"
 #include "List/SandboxListViewModel.h"
+#include "Persist/PersistSandboxViewModel.h"
 #include "SandboxControlsViewModel.h"
 #include "SandboxMetaDataViewModel.h"
 #include "List/SandboxColumnRegistry.h"
 #include "Templates/SharedPointer.h"
 
-namespace UE::FileSandboxUI
-{
-class IExternalSandboxActiveViewModel;
-}
+namespace UE::FileSandboxUI { class IExternalSandboxActiveViewModel; }
 
 namespace UE::SandboxedEditing
 {
@@ -34,8 +32,11 @@ struct FBrowserViewModels
 	/** Handles changing metadata of a sandbox. */
 	const TSharedRef<FSandboxMetaDataViewModel> MetaDataViewModel;
 	
+	/** Handles persisting the active sandbox. */
+	const TSharedRef<FPersistSandboxViewModel> PersistViewModel;
 	/** Handles leaving an active sandbox. */
 	const TSharedRef<FLeaveSandboxViewModel> LeaveViewModel;
+	
 	/** Used for widgets to switch states when the active sandbox changes. */
 	const TSharedRef<FActiveSandboxTrackerViewModel> ActiveSandboxTrackerViewModel;
 	/** Used for widgets that display info about the currently active sandbox. */
@@ -52,7 +53,8 @@ struct FBrowserViewModels
 		, ListViewModel(MakeShared<FSandboxListViewModel>(InModel, FilterSandboxViewModel, InSandboxColumns.ColumnBehaviors))
 		, ControlsViewModel(MakeShared<FSandboxControlsViewModel>(InModel))
 		, MetaDataViewModel(MakeShared<FSandboxMetaDataViewModel>(InModel))
-		, LeaveViewModel(MakeShared<FLeaveSandboxViewModel>(InModel))
+		, PersistViewModel(MakeShared<FPersistSandboxViewModel>(InModel))
+		, LeaveViewModel(MakeShared<FLeaveSandboxViewModel>(InModel, PersistViewModel))
 		, ActiveSandboxTrackerViewModel(MakeShared<FActiveSandboxTrackerViewModel>(InModel))
 		, ActiveSandboxDetailsViewModel(MakeShared<FActiveSandboxDetailsViewModel>(InModel, LeaveViewModel))
 		, ExternalSandboxViewModel(InExternalSandboxViewModel)
diff --git a/Engine/Plugins/Developer/Sandbox/SandboxedEditing/Source/SandboxedEditing/Private/Features/Browser/ViewModels/FileState/FileStateColumns.h b/Engine/Plugins/Developer/Sandbox/SandboxedEditing/Source/SandboxedEditing/Private/Features/Browser/ViewModels/FileState/FileStateColumns.h
new file mode 100644
--- /dev/null
+++ b/Engine/Plugins/Developer/Sandbox/SandboxedEditing/Source/SandboxedEditing/Private/Features/Browser/ViewModels/FileState/FileStateColumns.h
@@ -0,0 +1,20 @@
+﻿// Copyright Epic Games, Inc. All Rights Reserved.
+
+#pragma once
+
+#include "Templates/SharedPointer.h"
+#include "UObject/NameTypes.h"
+
+namespace UE::SandboxedEditing
+{
+class IFileStateColumnBehavior;
+class IFileStateColumnWidgetFactory;
+struct FFileStateItem;
+
+/** Needed for SLATE_ARGUMENT (the macro is written in such a way that TMap is not parsed correctly). */
+using FFileStateColumnFactoryMap = TMap<FName, TSharedRef<IFileStateColumnWidgetFactory>>;
+
+constexpr FLazyName PersistCheckboxColumn(TEXT("PersistCheckboxColumn"));
+constexpr FLazyName FileActionColumn(TEXT("FileActionColumn"));
+constexpr FLazyName FilePathColumn(TEXT("FilePathColumn"));
+}
diff --git a/Engine/Plugins/Developer/Sandbox/SandboxedEditing/Source/SandboxedEditing/Private/Features/Browser/ViewModels/FileState/FileStateItem.h b/Engine/Plugins/Developer/Sandbox/SandboxedEditing/Source/SandboxedEditing/Private/Features/Browser/ViewModels/FileState/FileStateItem.h
new file mode 100644
--- /dev/null
+++ b/Engine/Plugins/Developer/Sandbox/SandboxedEditing/Source/SandboxedEditing/Private/Features/Browser/ViewModels/FileState/FileStateItem.h
@@ -0,0 +1,24 @@
+﻿// Copyright Epic Games, Inc. All Rights Reserved.
+
+#pragma once
+
+#include "Containers/UnrealString.h"
+#include "Types/SandboxFileChange.h"
+
+namespace UE::SandboxedEditing
+{
+/** Represents an item displayed in IFileStateViewModel. */
+struct FFileStateItem
+{
+	/** Path to the non-sandboxed file. */
+	FString NonSandboxFile;
+	
+	/** The operation performed on the file. */
+	FileSandboxCore::ESandboxFileChange Action;
+
+	explicit FFileStateItem(FString InNonSandboxFile, FileSandboxCore::ESandboxFileChange InAction)
+		: NonSandboxFile(MoveTemp(InNonSandboxFile))
+		, Action(InAction)
+	{}
+};
+}
diff --git a/Engine/Plugins/Developer/Sandbox/SandboxedEditing/Source/SandboxedEditing/Private/Features/Browser/ViewModels/FileState/IFileStateColumnBehavior.h b/Engine/Plugins/Developer/Sandbox/SandboxedEditing/Source/SandboxedEditing/Private/Features/Browser/ViewModels/FileState/IFileStateColumnBehavior.h
new file mode 100644
--- /dev/null
+++ b/Engine/Plugins/Developer/Sandbox/SandboxedEditing/Source/SandboxedEditing/Private/Features/Browser/ViewModels/FileState/IFileStateColumnBehavior.h
@@ -0,0 +1,13 @@
+﻿// Copyright Epic Games, Inc. All Rights Reserved.
+
+#pragma once
+
+namespace UE::SandboxedEditing
+{
+/**
+ * Marker interface. For now there is no behaviour. 
+ * TODO UE-358389: Add interface for sorting
+ */
+class IFileStateColumnBehavior
+{};
+}
diff --git a/Engine/Plugins/Developer/Sandbox/SandboxedEditing/Source/SandboxedEditing/Private/Features/Browser/ViewModels/FileState/IFileStateColumnWidgetFactory.h b/Engine/Plugins/Developer/Sandbox/SandboxedEditing/Source/SandboxedEditing/Private/Features/Browser/ViewModels/FileState/IFileStateColumnWidgetFactory.h
new file mode 100644
--- /dev/null
+++ b/Engine/Plugins/Developer/Sandbox/SandboxedEditing/Source/SandboxedEditing/Private/Features/Browser/ViewModels/FileState/IFileStateColumnWidgetFactory.h
@@ -0,0 +1,35 @@
+﻿// Copyright Epic Games, Inc. All Rights Reserved.
+
+#pragma once
+
+#include "Widgets/Views/SHeaderRow.h"
+
+namespace UE::SandboxedEditing
+{
+struct FFileStateItem;
+
+/** Args for making a column widget in a row. */
+struct FMakeFileStateColumnWidgetArgs
+{
+	/** The row data for this column widget */
+	TSharedPtr<FFileStateItem> RowData;
+
+	explicit FMakeFileStateColumnWidgetArgs(TSharedPtr<FFileStateItem> InRowData)
+		: RowData(MoveTemp(InRowData))
+	{}
+};
+
+/** Interface for creating Slate widgets for a column. */
+class IFileStateColumnWidgetFactory
+{
+public:
+	
+	/** @return Info required to construct the column. */
+	virtual SHeaderRow::FColumn::FArguments MakeColumnArguments() = 0;
+	
+	/** @return Widget that should be put into the row for this column */
+	virtual TSharedRef<SWidget> MakeColumnWidget(const FMakeFileStateColumnWidgetArgs& InArgs) = 0;
+
+	virtual ~IFileStateColumnWidgetFactory() = default;
+};
+}
diff --git a/Engine/Plugins/Developer/Sandbox/SandboxedEditing/Source/SandboxedEditing/Private/Features/Browser/ViewModels/FileState/IFileStateViewModel.h b/Engine/Plugins/Developer/Sandbox/SandboxedEditing/Source/SandboxedEditing/Private/Features/Browser/ViewModels/FileState/IFileStateViewModel.h
new file mode 100644
--- /dev/null
+++ b/Engine/Plugins/Developer/Sandbox/SandboxedEditing/Source/SandboxedEditing/Private/Features/Browser/ViewModels/FileState/IFileStateViewModel.h
@@ -0,0 +1,33 @@
+﻿// Copyright Epic Games, Inc. All Rights Reserved.
+
+#pragma once
+
+#include "Delegates/Delegate.h"
+
+namespace UE::SandboxedEditing
+{
+struct FFileStateItem;
+
+/** 
+ * View model for UI that shows the files added / removed / edited in the sandbox. 
+ * 
+ * The implementation can vary based on where the UI is displayed, e.g.
+ * - Persist files (part of workflow for leaving sandbox)
+ * - Sandbox changed files (part of UI for displaying current sandbox details)
+ */
+class IFileStateViewModel
+{
+public:
+	
+	/** @return The items to display */
+	virtual const TArray<TSharedPtr<FFileStateItem>>& GetItems() const = 0;
+	
+	/** @return The columns that should be displayed. */
+	virtual TArray<FName> GetDisplayedColumns() const = 0;
+	
+	/** Invoked when the items change. */
+	virtual FSimpleMulticastDelegate& OnItemsChanged() = 0;
+	
+	virtual ~IFileStateViewModel() = default;
+};
+}
diff --git a/Engine/Plugins/Developer/Sandbox/SandboxedEditing/Source/SandboxedEditing/Private/Features/Browser/ViewModels/FileState/StaticFileStateViewModel.cpp b/Engine/Plugins/Developer/Sandbox/SandboxedEditing/Source/SandboxedEditing/Private/Features/Browser/ViewModels/FileState/StaticFileStateViewModel.cpp
new file mode 100644
--- /dev/null
+++ b/Engine/Plugins/Developer/Sandbox/SandboxedEditing/Source/SandboxedEditing/Private/Features/Browser/ViewModels/FileState/StaticFileStateViewModel.cpp
@@ -0,0 +1,21 @@
+﻿// Copyright Epic Games, Inc. All Rights Reserved.
+
+#include "StaticFileStateViewModel.h"
+
+namespace UE::SandboxedEditing
+{
+FStaticFileStateViewModel::FStaticFileStateViewModel(
+	TMap<FName, TSharedRef<IFileStateColumnBehavior>> InColumns,
+	TArray<TSharedPtr<FFileStateItem>> InItems
+	)
+	: Columns(MoveTemp(InColumns))
+	, Items(MoveTemp(InItems))
+{}
+
+TArray<FName> FStaticFileStateViewModel::GetDisplayedColumns() const
+{
+	TArray<FName> ColumnNames;
+	Columns.GenerateKeyArray(ColumnNames);
+	return ColumnNames;
+}
+}
diff --git a/Engine/Plugins/Developer/Sandbox/SandboxedEditing/Source/SandboxedEditing/Private/Features/Browser/ViewModels/FileState/StaticFileStateViewModel.h b/Engine/Plugins/Developer/Sandbox/SandboxedEditing/Source/SandboxedEditing/Private/Features/Browser/ViewModels/FileState/StaticFileStateViewModel.h
new file mode 100644
--- /dev/null
+++ b/Engine/Plugins/Developer/Sandbox/SandboxedEditing/Source/SandboxedEditing/Private/Features/Browser/ViewModels/FileState/StaticFileStateViewModel.h
@@ -0,0 +1,47 @@
+﻿// Copyright Epic Games, Inc. All Rights Reserved.
+
+#pragma once
+
+#include "IFileStateViewModel.h"
+
+namespace UE::SandboxedEditing
+{
+class IFileStateColumnBehavior;
+class FSandboxListItem;
+
+/**
+ * Is initialized with items at construction. 
+ * The item array is not mutated. (TODO except for sorting columns)
+ * This can serve as a base implementation for more dynamic models.
+ */
+class FStaticFileStateViewModel : public IFileStateViewModel
+{
+public:
+	
+	explicit FStaticFileStateViewModel(
+		TMap<FName, TSharedRef<IFileStateColumnBehavior>> InColumns,
+		TArray<TSharedPtr<FFileStateItem>> InItems
+		);
+	
+	//~ Begin IFileStateViewModel Interface
+	virtual const TArray<TSharedPtr<FFileStateItem>>& GetItems() const override { return Items; }
+	virtual TArray<FName> GetDisplayedColumns() const override;
+	virtual FSimpleMulticastDelegate& OnItemsChanged() override { return OnItemsChangedDelegate; }
+	//~ End IFileStateViewModel Interface
+
+protected:
+	
+	/** All columns that can ever be displayed. */
+	const TMap<FName, TSharedRef<IFileStateColumnBehavior>> Columns;
+	
+	/** The items to display. */
+	TArray<TSharedPtr<FFileStateItem>> Items;
+	
+	/**
+	 * Dummy to fulfill the interface requirements. Can be used by subclasses. 
+	 * TODO when sorting column is implemented, this can be invoked.
+	 */
+	FSimpleMulticastDelegate OnItemsChangedDelegate;
+};
+}
+
diff --git a/Engine/Plugins/Developer/Sandbox/SandboxedEditing/Source/SandboxedEditing/Private/Features/Browser/ViewModels/Leaving/AskUserToPersistWorkflow.h b/Engine/Plugins/Developer/Sandbox/SandboxedEditing/Source/SandboxedEditing/Private/Features/Browser/ViewModels/Leaving/AskUserToPersistWorkflow.h
new file mode 100644
--- /dev/null
+++ b/Engine/Plugins/Developer/Sandbox/SandboxedEditing/Source/SandboxedEditing/Private/Features/Browser/ViewModels/Leaving/AskUserToPersistWorkflow.h
@@ -0,0 +1,46 @@
+﻿// Copyright Epic Games, Inc. All Rights Reserved.
+
+#pragma once
+
+#include "Delegates/Delegate.h"
+#include "Templates/UnrealTemplate.h"
+
+namespace UE::SandboxedEditing
+{
+enum class EAskUserToPersistResult : uint8
+{
+	/** The user wants to select files to persist. */
+	ProceedToPersist,
+	/** The users wants to leave without persisting. */
+	LeaveWithoutPersist,
+	/** User decided to cancel the operation. */
+	Cancelled
+};
+
+/** Upon leaving the sandbox, asks the user whether they want to persist changes */
+class FAskUserToPersistWorkflow : public FNoncopyable
+{
+public:
+	
+	DECLARE_DELEGATE_OneParam(FCleanupWorkflow, EAskUserToPersistResult);
+	
+	explicit FAskUserToPersistWorkflow(FCleanupWorkflow InCleanupWorkflow)
+		: CleanupWorkflow(MoveTemp(InCleanupWorkflow))
+	{}
+	
+	/** Proceed to present the dialogue for persisting files. Called by view. */
+	void ProceedToPersist() { CleanupWorkflow.ExecuteIfBound(EAskUserToPersistResult::ProceedToPersist); }
+	
+	/** Leave the sandbox without persisting files. Called by view. */
+	void LeaveWithoutPersist() { CleanupWorkflow.ExecuteIfBound(EAskUserToPersistResult::LeaveWithoutPersist); }
+	
+	/** The user decided to cancel the operation. Called by view. */
+	void CancelLeave() { CleanupWorkflow.ExecuteIfBound(EAskUserToPersistResult::Cancelled); }
+	
+private:
+	
+	/** Invoked when the workflow ends. */
+	FCleanupWorkflow CleanupWorkflow;
+};
+}
+
diff --git a/Engine/Plugins/Developer/Sandbox/SandboxedEditing/Source/SandboxedEditing/Private/Features/Browser/ViewModels/Leaving/LeaveSandboxViewModel.cpp b/Engine/Plugins/Developer/Sandbox/SandboxedEditing/Source/SandboxedEditing/Private/Features/Browser/ViewModels/Leaving/LeaveSandboxViewModel.cpp
--- a/Engine/Plugins/Developer/Sandbox/SandboxedEditing/Source/SandboxedEditing/Private/Features/Browser/ViewModels/Leaving/LeaveSandboxViewModel.cpp
+++ b/Engine/Plugins/Developer/Sandbox/SandboxedEditing/Source/SandboxedEditing/Private/Features/Browser/ViewModels/Leaving/LeaveSandboxViewModel.cpp
@@ -2,17 +2,26 @@
 
 #include "LeaveSandboxViewModel.h"
 
+#include "Features/Browser/ViewModels/Persist/PersistWorkflowUserChoice.h"
+#include "Features/Browser/ViewModels/Persist/PersistSandboxViewModel.h"
 #include "Framework/Models/SandboxSystemModel.h"
 #include "Types/GatheredFileChanges.h"
 
 namespace UE::SandboxedEditing
 {
 FLeaveSandboxViewModel::FLeaveSandboxViewModel(
-	const TSharedRef<FSandboxSystemModel> InModel
+	const TSharedRef<FSandboxSystemModel> InModel, 
+	const TSharedRef<FPersistSandboxViewModel> InPersistViewModel
 	)
 	: Model(InModel)
+	, PersistViewModel(InPersistViewModel)
 {}
 
+FLeaveSandboxViewModel::~FLeaveSandboxViewModel()
+{
+	PersistViewModel->OnEndPersistWorkflow().RemoveAll(this);
+}
+
 void FLeaveSandboxViewModel::LeaveSandbox()
 {
 	if (!CanLeaveSandbox())
@@ -36,11 +45,20 @@ void FLeaveSandboxViewModel::LeaveSandbox()
 		OnCreateLeaveWithDirtyPackageWorkflowDelegate.Broadcast(*LeaveWithDirtyPackagesWorkflow);
 		return;
 	}
+	
+	const FileSandboxCore::FGatheredFileChanges FileChanges = Model->GatherFileChanges();
+	if (FileChanges.HasChanges())
+	{
+		AskUserToPersistWorkflow = MakeUnique<FAskUserToPersistWorkflow>(
+			FAskUserToPersistWorkflow::FCleanupWorkflow::CreateRaw(this, &FLeaveSandboxViewModel::OnUserDecidedWhetherToPersist)
+			);
+		AskUserToPersistWorkflowDelegate.Broadcast(*AskUserToPersistWorkflow);
+	}
 }
 
 bool FLeaveSandboxViewModel::CanLeaveSandbox(FText* OutReason) const
 {
-	return !LeaveWithDirtyPackagesWorkflow && Model->IsAllowedToLeaveSandbox(OutReason);
+	return !LeaveWithDirtyPackagesWorkflow && !PersistViewModel->GetPersistWorkflow() && Model->IsAllowedToLeaveSandbox(OutReason);
 }
 
 void FLeaveSandboxViewModel::OnDirtyPackagesHandled(FAskUserAboutDirtyPackagesWorkflow::EWorkflowResult InResult)
@@ -53,4 +71,38 @@ void FLeaveSandboxViewModel::OnDirtyPackagesHandled(FAskUserAboutDirtyPackagesWo
 		LeaveSandbox();
 	}
 }
+
+void FLeaveSandboxViewModel::OnUserDecidedWhetherToPersist(EAskUserToPersistResult InResult)
+{
+	AskUserToPersistWorkflow.Reset();
+
+	switch (InResult)
+	{
+	case EAskUserToPersistResult::ProceedToPersist: 
+		// Subscribe before starting the workflow because it can complete "immediately" (due to modal window).
+		PersistViewModel->OnEndPersistWorkflow().AddRaw(this, &FLeaveSandboxViewModel::OnPersistOperationEnded);
+		
+		PersistViewModel->StartPersistWorkflow(Model->GatherFileChanges());
+		
+		// If the workflow immediately ends (due to modal window) or fails, unsubscribe again
+		if (!PersistViewModel->GetPersistWorkflow())
+		{
+			PersistViewModel->OnEndPersistWorkflow().RemoveAll(this);
+		}
+		break;
+	case EAskUserToPersistResult::LeaveWithoutPersist: Model->LeaveSandbox(); break;
+	case EAskUserToPersistResult::Cancelled: break;
+	default: break;
+	}
+}
+
+void FLeaveSandboxViewModel::OnPersistOperationEnded(EPersistWorkflowAction InAction) const
+{
+	if (InAction != EPersistWorkflowAction::Cancelled)
+	{
+		Model->LeaveSandbox();
+	}
+	
+	PersistViewModel->OnEndPersistWorkflow().RemoveAll(this);
+}
 }
diff --git a/Engine/Plugins/Developer/Sandbox/SandboxedEditing/Source/SandboxedEditing/Private/Features/Browser/ViewModels/Leaving/LeaveSandboxViewModel.h b/Engine/Plugins/Developer/Sandbox/SandboxedEditing/Source/SandboxedEditing/Private/Features/Browser/ViewModels/Leaving/LeaveSandboxViewModel.h
--- a/Engine/Plugins/Developer/Sandbox/SandboxedEditing/Source/SandboxedEditing/Private/Features/Browser/ViewModels/Leaving/LeaveSandboxViewModel.h
+++ b/Engine/Plugins/Developer/Sandbox/SandboxedEditing/Source/SandboxedEditing/Private/Features/Browser/ViewModels/Leaving/LeaveSandboxViewModel.h
@@ -3,8 +3,10 @@
 #pragma once
 
 #include "AskUserAboutDirtyPackagesWorkflow.h"
+#include "AskUserToPersistWorkflow.h"
 #include "Delegates/Delegate.h"
 #include "Delegates/DelegateCombinations.h"
+#include "Features/Browser/ViewModels/Persist/PersistSandboxViewModel.h"
 #include "Templates/SharedPointer.h"
 #include "Templates/UnrealTemplate.h"
 
@@ -13,7 +15,8 @@ namespace UE::SandboxedEditing
 class FSandboxSystemModel;
 enum class EPersistWorkflowAction : uint8;
 
-DECLARE_MULTICAST_DELEGATE_OneParam(FAskUserAboutDirtyPackagesWorkflowDelegate, FAskUserAboutDirtyPackagesWorkflow&);
+DECLARE_MULTICAST_DELEGATE_OneParam(FLeaveWithDirtyPackagesWorkflowDelegate, FAskUserAboutDirtyPackagesWorkflow&);
+DECLARE_MULTICAST_DELEGATE_OneParam(FAskUserToPersistWorkflowDelegate, FAskUserToPersistWorkflow&);
 
 /**
  * Handles the userflow of leaving a sandbox. 
@@ -29,8 +32,10 @@ class FLeaveSandboxViewModel : public FNoncopyable
 public:
 	
 	explicit FLeaveSandboxViewModel(
-		const TSharedRef<FSandboxSystemModel> InModel
+		const TSharedRef<FSandboxSystemModel> InModel,
+		const TSharedRef<FPersistSandboxViewModel> InPersistViewModel
 		);
+	~FLeaveSandboxViewModel();
 
 	/** Leaves the active sandbox. */
 	void LeaveSandbox();
@@ -39,23 +44,37 @@ public:
 	
 	/** @return Gets the current workflow */
 	FAskUserAboutDirtyPackagesWorkflow* GetLeaveWithDirtyPackageWorkflow() const { return LeaveWithDirtyPackagesWorkflow.Get(); }
+	/** @return Gets the current workflow */
+	FAskUserToPersistWorkflow* GetAskUserToPersistWorkflow() const { return AskUserToPersistWorkflow.Get(); }
 	
 	/** Invoked when the user should be asked about what to do with dirty packages. */
-	FAskUserAboutDirtyPackagesWorkflowDelegate& OnCreateLeaveWithDirtyPackageWorkflow() { return OnCreateLeaveWithDirtyPackageWorkflowDelegate; }
+	FLeaveWithDirtyPackagesWorkflowDelegate& OnCreateLeaveWithDirtyPackageWorkflow() { return OnCreateLeaveWithDirtyPackageWorkflowDelegate; }
+	/** Invoked when the user should be asked whether they want to persist any files. */
+	FAskUserToPersistWorkflowDelegate& OnAskUserToPersistWorkflow() { return AskUserToPersistWorkflowDelegate; }
 	
 private:
 	
 	/** The model used to interact with the sandbox system. */
 	const TSharedRef<FSandboxSystemModel> Model;
 
+	/** Used to go through the persist workflow. */
+	const TSharedRef<FPersistSandboxViewModel> PersistViewModel;
 
 	/** Handles dirty packages when leaving. */
 	TUniquePtr<FAskUserAboutDirtyPackagesWorkflow> LeaveWithDirtyPackagesWorkflow;
+	/** Asks the user whether they want to persist, leave without persisting, or cancel. */
+	TUniquePtr<FAskUserToPersistWorkflow> AskUserToPersistWorkflow;
 	
 	/** Invoked when the user should be asked about what to do with dirty packages. */
-	FAskUserAboutDirtyPackagesWorkflowDelegate OnCreateLeaveWithDirtyPackageWorkflowDelegate;
+	FLeaveWithDirtyPackagesWorkflowDelegate OnCreateLeaveWithDirtyPackageWorkflowDelegate;
+	/** Invoked when the user should be asked whether they want to persist any files. */
+	FAskUserToPersistWorkflowDelegate AskUserToPersistWorkflowDelegate;
 	
 	/** Invoked when LeaveWithDirtyPackagesWorkflow finishes. */
 	void OnDirtyPackagesHandled(FAskUserAboutDirtyPackagesWorkflow::EWorkflowResult InResult);
+	/** Invoked when the FAskUserToPersistWorkflow ends.*/
+	void OnUserDecidedWhetherToPersist(EAskUserToPersistResult InResult);
+	/** Invoked when the FPersistSandboxWorkflow has ended. */
+	void OnPersistOperationEnded(EPersistWorkflowAction InAction) const;
 };
 }
diff --git a/Engine/Plugins/Developer/Sandbox/SandboxedEditing/Source/SandboxedEditing/Private/Features/Browser/ViewModels/Persist/PersistOperationViewModel.cpp b/Engine/Plugins/Developer/Sandbox/SandboxedEditing/Source/SandboxedEditing/Private/Features/Browser/ViewModels/Persist/PersistOperationViewModel.cpp
new file mode 100644
--- /dev/null
+++ b/Engine/Plugins/Developer/Sandbox/SandboxedEditing/Source/SandboxedEditing/Private/Features/Browser/ViewModels/Persist/PersistOperationViewModel.cpp
@@ -0,0 +1,67 @@
+﻿// Copyright Epic Games, Inc. All Rights Reserved.
+
+#include "PersistOperationViewModel.h"
+
+#include "Algo/AllOf.h"
+#include "Algo/AnyOf.h"
+#include "PersistSandboxWorkflow.h"
+
+namespace UE::SandboxedEditing
+{
+FPersistOperationViewModel::FPersistOperationViewModel(FPersistSandboxWorkflow& InWorkflow)
+	: Workflow(InWorkflow)
+	, FilesMarkedForPersist(true, Workflow.GetPersistableFiles().NonSandboxPaths.Num())
+{}
+
+void FPersistOperationViewModel::ConfirmPersist() const
+{
+	TArray<FString> FilesToPersist;
+	const TArray<FString>& NonSandboxPaths = GetPersistableFiles().NonSandboxPaths;
+	for (int32 Index = 0; Index < NonSandboxPaths.Num(); ++Index)
+	{
+		if (IsFileMarkedForPersist(Index))
+		{
+			FilesToPersist.Add(NonSandboxPaths[Index]);
+		}
+	}
+	
+	Workflow.ConfirmPersist(
+		MoveTemp(FilesToPersist)
+		);
+}
+
+void FPersistOperationViewModel::CancelPersist() const
+{
+	Workflow.CancelPersist();
+}
+
+void FPersistOperationViewModel::SetFilePersisted(int32 InIndex, bool bShouldPersist)
+{
+	FilesMarkedForPersist[InIndex] = bShouldPersist;
+}
+
+bool FPersistOperationViewModel::IsFileMarkedForPersist(int32 InIndex) const
+{
+	return FilesMarkedForPersist[InIndex];
+}
+
+void FPersistOperationViewModel::SetAllFilesPersisted(bool bAllFiles)
+{
+	FilesMarkedForPersist = TBitArray(bAllFiles, Workflow.GetPersistableFiles().NonSandboxPaths.Num());
+}
+
+bool FPersistOperationViewModel::AreAllFilesMarkedForPersist() const
+{
+	return Algo::AllOf(FilesMarkedForPersist);
+}
+
+bool FPersistOperationViewModel::AreAnyFilesMarkedForPersist() const
+{
+	return Algo::AnyOf(FilesMarkedForPersist);
+}
+
+const FileSandboxCore::FGatheredFileChanges& FPersistOperationViewModel::GetPersistableFiles() const
+{
+	return Workflow.GetPersistableFiles();
+}
+}
diff --git a/Engine/Plugins/Developer/Sandbox/SandboxedEditing/Source/SandboxedEditing/Private/Features/Browser/ViewModels/Persist/PersistOperationViewModel.h b/Engine/Plugins/Developer/Sandbox/SandboxedEditing/Source/SandboxedEditing/Private/Features/Browser/ViewModels/Persist/PersistOperationViewModel.h
new file mode 100644
--- /dev/null
+++ b/Engine/Plugins/Developer/Sandbox/SandboxedEditing/Source/SandboxedEditing/Private/Features/Browser/ViewModels/Persist/PersistOperationViewModel.h
@@ -0,0 +1,57 @@
+﻿// Copyright Epic Games, Inc. All Rights Reserved.
+
+#pragma once
+
+#include "Containers/BitArray.h"
+#include "HAL/Platform.h"
+#include "Types/GatheredFileChanges.h"
+
+namespace UE::SandboxedEditing
+{
+class FPersistSandboxWorkflow;
+
+/** View model for UI that lets the user click items they want to persist. */
+class FPersistOperationViewModel : public FNoncopyable
+{
+public:
+	
+	explicit FPersistOperationViewModel(FPersistSandboxWorkflow& InWorkflow UE_LIFETIMEBOUND);
+	
+	/** Persists the files selected so far. Any follow-up operation, such as leaving the sandbox, will be performed. */
+	void ConfirmPersist() const;
+	
+	/** Cancels the workflow. No follow-up operation, such as leaving the sandbox, will be performed. */
+	void CancelPersist() const;
+	
+	/** 
+	 * @param InIndex Index to GetPersistableFiles().NonSandboxPaths
+	 * @param bShouldPersist Whether the files should be persisted
+	 */
+	void SetFilePersisted(int32 InIndex, bool bShouldPersist);
+	/** 
+	 * @param InIndex Index to GetPersistableFiles().NonSandboxPaths
+	 * @return Whether the file is marked to be persisted 
+	 */
+	bool IsFileMarkedForPersist(int32 InIndex) const;
+	
+	/** Marks all files to be persisted. */
+	void SetAllFilesPersisted(bool bAllFiles);
+	
+	/** @return Whether all files are marked to be persisted. */
+	bool AreAllFilesMarkedForPersist() const;
+	/** @return Whether any files are marked for persist. */
+	bool AreAnyFilesMarkedForPersist() const;
+	
+	/** @return The modified files */
+	const FileSandboxCore::FGatheredFileChanges& GetPersistableFiles() const;
+	
+private:
+	
+	/** The workflow to forward the persist and cancel events to. */
+	FPersistSandboxWorkflow& Workflow;
+	
+	/** The files the user has selected to persist. Each index corresponds to PersistableFiles.NonSandboxPaths. */
+	TBitArray<> FilesMarkedForPersist;
+};
+}
+
diff --git a/Engine/Plugins/Developer/Sandbox/SandboxedEditing/Source/SandboxedEditing/Private/Features/Browser/ViewModels/Persist/PersistSandboxViewModel.cpp b/Engine/Plugins/Developer/Sandbox/SandboxedEditing/Source/SandboxedEditing/Private/Features/Browser/ViewModels/Persist/PersistSandboxViewModel.cpp
new file mode 100644
--- /dev/null
+++ b/Engine/Plugins/Developer/Sandbox/SandboxedEditing/Source/SandboxedEditing/Private/Features/Browser/ViewModels/Persist/PersistSandboxViewModel.cpp
@@ -0,0 +1,42 @@
+﻿// Copyright Epic Games, Inc. All Rights Reserved.
+
+#include "PersistSandboxViewModel.h"
+
+#include "PersistWorkflowUserChoice.h"
+
+namespace UE::SandboxedEditing
+{
+FPersistSandboxViewModel::FPersistSandboxViewModel(const TSharedRef<FSandboxSystemModel>& InModel)
+	: Model(InModel)
+{}
+
+bool FPersistSandboxViewModel::StartPersistWorkflow(FileSandboxCore::FGatheredFileChanges InFileChanges)
+{
+	if (CanStartPersistWorkflow())
+	{
+		ActivePersistWorkflow = MakeUnique<FPersistSandboxWorkflow>(Model, MoveTemp(InFileChanges),
+			FPersistSandboxWorkflow::FCleanupWorkflow::CreateRaw(this, &FPersistSandboxViewModel::OnCleanupPersistWorkflow)
+			);
+		OnStartPersistWorkflowDelegate.Broadcast(*ActivePersistWorkflow);
+		return true;
+	}
+	
+	return false;
+}
+
+bool FPersistSandboxViewModel::CanStartPersistWorkflow() const
+{
+	return !ActivePersistWorkflow.IsValid();
+}
+
+void FPersistSandboxViewModel::OnCleanupPersistWorkflow(const FPersistWorkflowUserChoice& InPersistResult)
+{
+	if (InPersistResult.UserAction == EPersistWorkflowAction::Persist)
+	{
+		Model->PersistFiles(*InPersistResult.FilesToPersist);
+	}
+	
+	ActivePersistWorkflow.Reset();
+	OnEndPersistWorkflowDelegate.Broadcast(InPersistResult.UserAction);
+}
+}
diff --git a/Engine/Plugins/Developer/Sandbox/SandboxedEditing/Source/SandboxedEditing/Private/Features/Browser/ViewModels/Persist/PersistSandboxViewModel.h b/Engine/Plugins/Developer/Sandbox/SandboxedEditing/Source/SandboxedEditing/Private/Features/Browser/ViewModels/Persist/PersistSandboxViewModel.h
new file mode 100644
--- /dev/null
+++ b/Engine/Plugins/Developer/Sandbox/SandboxedEditing/Source/SandboxedEditing/Private/Features/Browser/ViewModels/Persist/PersistSandboxViewModel.h
@@ -0,0 +1,53 @@
+﻿// Copyright Epic Games, Inc. All Rights Reserved.
+
+#pragma once
+
+#include "PersistSandboxWorkflow.h"
+#include "Framework/Models/SandboxSystemModel.h"
+#include "Templates/UnrealTemplate.h"
+
+namespace UE::SandboxedEditing
+{
+enum class EPersistWorkflowAction : uint8;
+DECLARE_MULTICAST_DELEGATE_OneParam(FPersistSandboxWorkflowDelegate, FPersistSandboxWorkflow&)
+DECLARE_MULTICAST_DELEGATE_OneParam(FPersistWorkflowActionDelegate, EPersistWorkflowAction);
+
+/** Handles the logic of persisting a sandbox. */
+class FPersistSandboxViewModel : public FNoncopyable
+{
+public:
+	
+	explicit FPersistSandboxViewModel(const TSharedRef<FSandboxSystemModel>& InModel);
+
+	/** Starts the persist workflow. */
+	bool StartPersistWorkflow(FileSandboxCore::FGatheredFileChanges InFileChanges);
+	
+	/** @return Whether a persist workflow can be started */
+	bool CanStartPersistWorkflow() const;
+	
+	/** @return The current persist workflow */
+	FPersistSandboxWorkflow* GetPersistWorkflow() const { return ActivePersistWorkflow.Get(); }
+	
+	/** Invoked when the persist workflow is started. */
+	FPersistSandboxWorkflowDelegate& OnStartPersistWorkflow() { return OnStartPersistWorkflowDelegate; }
+	/** Invoked when the persist workflow has ended. */
+	FPersistWorkflowActionDelegate& OnEndPersistWorkflow() { return OnEndPersistWorkflowDelegate; }
+	
+private:
+	
+	/** Used to persist changes via the sandbox system. */
+	const TSharedRef<FSandboxSystemModel> Model;
+	
+	/** The active persist workflow. */
+	TUniquePtr<FPersistSandboxWorkflow> ActivePersistWorkflow;
+	
+	/** Invoked when the persist workflow is started. */
+	FPersistSandboxWorkflowDelegate OnStartPersistWorkflowDelegate;
+	/** Invoked when the persist workflow has ended. */
+	FPersistWorkflowActionDelegate OnEndPersistWorkflowDelegate;
+	
+	/** Handles the workflow ending, possibly persisting InFilesToPersist. */
+	void OnCleanupPersistWorkflow(const FPersistWorkflowUserChoice& InPersistResult);
+};
+}
+
diff --git a/Engine/Plugins/Developer/Sandbox/SandboxedEditing/Source/SandboxedEditing/Private/Features/Browser/ViewModels/Persist/PersistSandboxWorkflow.cpp b/Engine/Plugins/Developer/Sandbox/SandboxedEditing/Source/SandboxedEditing/Private/Features/Browser/ViewModels/Persist/PersistSandboxWorkflow.cpp
new file mode 100644
--- /dev/null
+++ b/Engine/Plugins/Developer/Sandbox/SandboxedEditing/Source/SandboxedEditing/Private/Features/Browser/ViewModels/Persist/PersistSandboxWorkflow.cpp
@@ -0,0 +1,40 @@
+﻿// Copyright Epic Games, Inc. All Rights Reserved.
+
+#include "PersistSandboxWorkflow.h"
+
+#include "PersistWorkflowUserChoice.h"
+#include "Algo/Transform.h"
+#include "Types/GatheredFileChanges.h"
+
+namespace UE::SandboxedEditing
+{
+FPersistSandboxWorkflow::FPersistSandboxWorkflow(
+	const TSharedRef<FSandboxSystemModel>& InModel, FileSandboxCore::FGatheredFileChanges InChanges,
+	FCleanupWorkflow InOnWorkflowEnded
+	)
+	: Model(InModel)
+	, OnCleanupWorkflowDelegate(MoveTemp(InOnWorkflowEnded))
+	, PersistableFiles(InChanges)
+{}
+
+void FPersistSandboxWorkflow::ConfirmPersist(TArray<FString> InFilesToPersist) const
+{
+	EndWorkflow(InFilesToPersist.IsEmpty() 
+		? FPersistWorkflowUserChoice::MakePersistNone()
+		: FPersistWorkflowUserChoice::MakePersist(MoveTemp(InFilesToPersist))
+		);
+}
+
+void FPersistSandboxWorkflow::CancelPersist() const
+{
+	EndWorkflow(
+		FPersistWorkflowUserChoice::MakeCancelled()
+		);
+}
+
+void FPersistSandboxWorkflow::EndWorkflow(const FPersistWorkflowUserChoice& InResult) const
+{
+	OnWorkflowEndedDelegate.Broadcast();
+	OnCleanupWorkflowDelegate.ExecuteIfBound(InResult);
+}
+}
diff --git a/Engine/Plugins/Developer/Sandbox/SandboxedEditing/Source/SandboxedEditing/Private/Features/Browser/ViewModels/Persist/PersistSandboxWorkflow.h b/Engine/Plugins/Developer/Sandbox/SandboxedEditing/Source/SandboxedEditing/Private/Features/Browser/ViewModels/Persist/PersistSandboxWorkflow.h
new file mode 100644
--- /dev/null
+++ b/Engine/Plugins/Developer/Sandbox/SandboxedEditing/Source/SandboxedEditing/Private/Features/Browser/ViewModels/Persist/PersistSandboxWorkflow.h
@@ -0,0 +1,67 @@
+﻿// Copyright Epic Games, Inc. All Rights Reserved.
+
+#pragma once
+
+#include "Delegates/Delegate.h"
+#include "Delegates/DelegateCombinations.h"
+#include "Templates/SharedPointer.h"
+#include "Templates/UnrealTemplate.h"
+#include "Types/GatheredFileChanges.h"
+
+namespace UE::SandboxedEditing
+{
+class FSandboxSystemModel;
+struct FPersistWorkflowUserChoice;
+
+enum class EPersistWorkflowResult : uint8
+{
+	/** Files were persisted, or the user decided to confirm without any files to persist. */
+	Persisted,
+	/** User decided to cancel the operation. Do not continue with follow-up operations, such as leaving the sandbox. */
+	Cancelled
+};
+
+/**
+ * Represents a temporary application state in which user operation is choosing files to persist. 
+ * Effectively, this is a view-model representation of the pop-up dialogue, in which files to persist are chosen.
+ */
+class FPersistSandboxWorkflow : public FNoncopyable
+{
+public:
+	
+	/** @param FilesToPersist The files the user selected to persist. Unset if the operation is cancelled. Can be set but empty. */
+	DECLARE_DELEGATE_OneParam(FCleanupWorkflow, const FPersistWorkflowUserChoice&);
+	
+	explicit FPersistSandboxWorkflow(
+		const TSharedRef<FSandboxSystemModel>& InModel, FileSandboxCore::FGatheredFileChanges InChanges,
+		FCleanupWorkflow InOnWorkflowEnded
+		);
+	
+	/** Persists the files selected so far. Any follow-up operation, such as leaving the sandbox, will be performed. */
+	void ConfirmPersist(TArray<FString> InFilesToPersist) const;
+	
+	/** Cancels the workflow. No follow-up operation, such as leaving the sandbox, will be performed. */
+	void CancelPersist() const;
+	
+	/** @return The modified files */
+	const FileSandboxCore::FGatheredFileChanges& GetPersistableFiles() const { return PersistableFiles; }
+	
+	/** Invoked when the workflow ends. */
+	FSimpleMulticastDelegate& OnWorkflowEnded() { return OnWorkflowEndedDelegate; }
+
+private:
+	
+	/** The model used to interact with the sandbox system. */
+	const TSharedRef<FSandboxSystemModel> Model;
+	
+	/** Invoked when the workflow ends. */
+	FSimpleMulticastDelegate OnWorkflowEndedDelegate;
+	/** Invoked when the workflow ends. */
+	const FCleanupWorkflow OnCleanupWorkflowDelegate;
+	
+	/** The files that can be persisted. */
+	const FileSandboxCore::FGatheredFileChanges PersistableFiles;
+	
+	void EndWorkflow(const FPersistWorkflowUserChoice& InResult) const;
+};
+}
diff --git a/Engine/Plugins/Developer/Sandbox/SandboxedEditing/Source/SandboxedEditing/Private/Features/Browser/ViewModels/Persist/PersistWorkflowUserChoice.h b/Engine/Plugins/Developer/Sandbox/SandboxedEditing/Source/SandboxedEditing/Private/Features/Browser/ViewModels/Persist/PersistWorkflowUserChoice.h
new file mode 100644
--- /dev/null
+++ b/Engine/Plugins/Developer/Sandbox/SandboxedEditing/Source/SandboxedEditing/Private/Features/Browser/ViewModels/Persist/PersistWorkflowUserChoice.h
@@ -0,0 +1,40 @@
+﻿// Copyright Epic Games, Inc. All Rights Reserved.
+
+#pragma once
+
+#include "Containers/UnrealString.h"
+#include "HAL/Platform.h"
+#include "Misc/Optional.h"
+
+namespace UE::SandboxedEditing
+{
+/** Action to take when FPersistSandboxWorkflow completes. */
+enum class EPersistWorkflowAction : uint8
+{
+	/** Files were selected for persist. Any follow-up operation, such as leaving, should be performed. */
+	Persist,
+	
+	/** User explicitly choose not to persist any of the files. Any follow-up operation, such as leaving, should be performed. */
+	PersistNone,
+	
+	/** User decided to cancel the operation. Do not continue with follow-up operations, such as leaving the sandbox. */
+	Cancelled
+};
+
+/** Describes the user decision for FPersistSandboxWorkflow. */
+struct FPersistWorkflowUserChoice
+{
+	/** Determines what action the user chose to take. */
+	EPersistWorkflowAction UserAction;
+	
+	/** Only set if UserAction == EPersistWorkflowAction::Persist. */
+	TOptional<TArray<FString>> FilesToPersist;
+	
+	static FPersistWorkflowUserChoice MakePersist(TArray<FString> InFileToPersist)
+	{
+		return FPersistWorkflowUserChoice(EPersistWorkflowAction::Persist, MoveTemp(InFileToPersist));
+	}
+	static FPersistWorkflowUserChoice MakePersistNone() { return FPersistWorkflowUserChoice(EPersistWorkflowAction::PersistNone); }
+	static FPersistWorkflowUserChoice MakeCancelled() { return FPersistWorkflowUserChoice(EPersistWorkflowAction::Cancelled); }
+};
+}
diff --git a/Engine/Plugins/Developer/Sandbox/SandboxedEditing/Source/SandboxedEditing/Private/Features/Browser/Views/AskUserToPersistView.cpp b/Engine/Plugins/Developer/Sandbox/SandboxedEditing/Source/SandboxedEditing/Private/Features/Browser/Views/AskUserToPersistView.cpp
new file mode 100644
--- /dev/null
+++ b/Engine/Plugins/Developer/Sandbox/SandboxedEditing/Source/SandboxedEditing/Private/Features/Browser/Views/AskUserToPersistView.cpp
@@ -0,0 +1,40 @@
+﻿// Copyright Epic Games, Inc. All Rights Reserved.
+
+#include "AskUserToPersistView.h"
+
+#include "Features/Browser/ViewModels/Leaving/LeaveSandboxViewModel.h"
+#include "Internationalization/Internationalization.h"
+#include "Misc/MessageDialog.h"
+
+#define LOCTEXT_NAMESPACE "FAskUserToPersistView"
+
+namespace UE::SandboxedEditing
+{
+FAskUserToPersistView::FAskUserToPersistView(const TSharedRef<FLeaveSandboxViewModel>& InLeaveViewModel)
+	: LeaveViewModel(InLeaveViewModel)
+{
+	LeaveViewModel->OnAskUserToPersistWorkflow().AddRaw(this, &FAskUserToPersistView::ShowNotification);
+}
+
+FAskUserToPersistView::~FAskUserToPersistView()
+{
+	LeaveViewModel->OnAskUserToPersistWorkflow().RemoveAll(this);
+}
+
+void FAskUserToPersistView::ShowNotification(FAskUserToPersistWorkflow& InWorkflow) const
+{
+	const EAppReturnType::Type Result = FMessageDialog::Open(
+		EAppMsgType::YesNoCancel, EAppReturnType::No,
+		LOCTEXT("Question", "You are about to leave a sandbox containing changes. Do you want to persist the changes?")
+		);
+	switch (Result)
+	{
+	case EAppReturnType::Yes: InWorkflow.ProceedToPersist(); break;
+	case EAppReturnType::No: InWorkflow.LeaveWithoutPersist(); break;
+	case EAppReturnType::Cancel: InWorkflow.CancelLeave(); break;
+	default: checkNoEntry(); InWorkflow.CancelLeave();
+	}
+}
+}
+
+#undef LOCTEXT_NAMESPACE
\ No newline at end of file
diff --git a/Engine/Plugins/Developer/Sandbox/SandboxedEditing/Source/SandboxedEditing/Private/Features/Browser/Views/AskUserToPersistView.h b/Engine/Plugins/Developer/Sandbox/SandboxedEditing/Source/SandboxedEditing/Private/Features/Browser/Views/AskUserToPersistView.h
new file mode 100644
--- /dev/null
+++ b/Engine/Plugins/Developer/Sandbox/SandboxedEditing/Source/SandboxedEditing/Private/Features/Browser/Views/AskUserToPersistView.h
@@ -0,0 +1,27 @@
+﻿// Copyright Epic Games, Inc. All Rights Reserved.
+
+#pragma once
+
+#include "Templates/SharedPointer.h"
+
+namespace UE::SandboxedEditing
+{
+class FAskUserToPersistWorkflow;
+class FLeaveSandboxViewModel;
+
+/** Visualizes FAskUserToPersistWorkflow by displaying a dialogue. */
+class FAskUserToPersistView
+{
+public:
+	
+	explicit FAskUserToPersistView(const TSharedRef<FLeaveSandboxViewModel>& InLeaveViewModel);
+	~FAskUserToPersistView();
+	
+private:
+	
+	/** Notifies when FAskUserToPersistWorkflow starts.  */
+	const TSharedRef<FLeaveSandboxViewModel> LeaveViewModel;
+
+	void ShowNotification(FAskUserToPersistWorkflow& InWorkflow) const;
+};
+}
diff --git a/Engine/Plugins/Developer/Sandbox/SandboxedEditing/Source/SandboxedEditing/Private/Features/Browser/Views/BrowserViews.h b/Engine/Plugins/Developer/Sandbox/SandboxedEditing/Source/SandboxedEditing/Private/Features/Browser/Views/BrowserViews.h
new file mode 100644
--- /dev/null
+++ b/Engine/Plugins/Developer/Sandbox/SandboxedEditing/Source/SandboxedEditing/Private/Features/Browser/Views/BrowserViews.h
@@ -0,0 +1,37 @@
+﻿// Copyright Epic Games, Inc. All Rights Reserved.
+
+#pragma once
+
+#include "AskUserAboutDirtyPackagesView.h"
+#include "AskUserToPersistView.h"
+#include "PersistSandboxView.h"
+#include "Templates/SharedPointer.h"
+
+namespace UE::SandboxedEditing
+{
+class FLeaveSandboxViewModel;
+
+/** 
+ * Holds the views for notifications that happen in Sandboxed Editing.
+ * @note Only contains "explicit" views. Of course, widgets, e.g. SSandboxedEditingRoot, implicitly classify as views. Those are not held here though.
+ */
+struct FBrowserViews
+{
+	/** When leaving with in-memory changes, this displays a toaster. */
+	const TSharedRef<FAskUserAboutDirtyPackagesView> AskUserAboutDirtyPackagesView;
+	
+	/** When leaving, this asks the user whether they want to persist files. */
+	const TSharedRef<FAskUserToPersistView> AskUserToPersistView;
+	
+	/** Visualizes FPersistSandboxWorkflow. */
+	const TSharedRef<FPersistSandboxView> PersistSandboxView;
+	
+	explicit FBrowserViews(
+		const TSharedRef<FPersistSandboxViewModel>& InPersistViewModel, const TSharedRef<FLeaveSandboxViewModel>& InLeaveViewModel
+		)
+		: AskUserAboutDirtyPackagesView(MakeShared<FAskUserAboutDirtyPackagesView>(InLeaveViewModel))
+		, AskUserToPersistView(MakeShared<FAskUserToPersistView>(InLeaveViewModel))
+		, PersistSandboxView(MakeShared<FPersistSandboxView>(InPersistViewModel))
+	{}
+};
+}
diff --git a/Engine/Plugins/Developer/Sandbox/SandboxedEditing/Source/SandboxedEditing/Private/Features/Browser/Views/PersistSandboxView.cpp b/Engine/Plugins/Developer/Sandbox/SandboxedEditing/Source/SandboxedEditing/Private/Features/Browser/Views/PersistSandboxView.cpp
new file mode 100644
--- /dev/null
+++ b/Engine/Plugins/Developer/Sandbox/SandboxedEditing/Source/SandboxedEditing/Private/Features/Browser/Views/PersistSandboxView.cpp
@@ -0,0 +1,56 @@
+﻿// Copyright Epic Games, Inc. All Rights Reserved.
+
+#include "PersistSandboxView.h"
+
+#include "Features/Browser/ViewModels/Persist/PersistOperationViewModel.h"
+#include "Features/Browser/ViewModels/Persist/PersistSandboxViewModel.h"
+#include "Features/Browser/Widgets/Sandboxed/SPersistOperationWidget.h"
+#include "Framework/Application/SlateApplication.h"
+#include "Widgets/SWindow.h"
+
+#define LOCTEXT_NAMESPACE "FPersistSandboxView"
+
+namespace UE::SandboxedEditing
+{
+FPersistSandboxView::FPersistSandboxView(const TSharedRef<FPersistSandboxViewModel>& InPersistViewModel)
+	: PersistViewModel(InPersistViewModel)
+{
+	PersistViewModel->OnStartPersistWorkflow().AddRaw(this, &FPersistSandboxView::ShowPersistWorkflow);
+}
+
+FPersistSandboxView::~FPersistSandboxView()
+{
+	PersistViewModel->OnStartPersistWorkflow().RemoveAll(this);
+}
+
+void FPersistSandboxView::ShowPersistWorkflow(FPersistSandboxWorkflow& InWorkflow)
+{
+	const TSharedRef<SWindow> NewWindow = SNew(SWindow)
+		.Title(LOCTEXT("PersistSubmitWindowTitle", "Persist Sandbox"))
+		.SizingRule(ESizingRule::UserSized)
+		.ClientSize(FVector2D(600, 600))
+		.SupportsMaximize(true)
+		.SupportsMinimize(false);
+	
+	const TSharedRef<FPersistOperationViewModel> ViewModel = MakeShared<FPersistOperationViewModel>(InWorkflow);
+	NewWindow->SetContent(
+		SNew(SPersistOperationWidget, ViewModel)
+	);
+	
+	InWorkflow.OnWorkflowEnded().AddLambda([WeakWindow = NewWindow.ToWeakPtr()]
+	{
+		const TSharedPtr<SWindow> WindowPin = WeakWindow.Pin();
+		WindowPin->RequestDestroyWindow();
+	});
+	FSlateApplication::Get().AddModalWindow(NewWindow, nullptr);
+	
+	// If the user just closed the window, treat it as cancelling the workflow.
+	// The workflow may be destroyed as the window may have confirmed the persist operation.
+	if (FPersistSandboxWorkflow* Workflow = PersistViewModel->GetPersistWorkflow())
+	{
+		Workflow->CancelPersist();
+	}
+}
+}
+
+#undef LOCTEXT_NAMESPACE
\ No newline at end of file
diff --git a/Engine/Plugins/Developer/Sandbox/SandboxedEditing/Source/SandboxedEditing/Private/Features/Browser/Views/PersistSandboxView.h b/Engine/Plugins/Developer/Sandbox/SandboxedEditing/Source/SandboxedEditing/Private/Features/Browser/Views/PersistSandboxView.h
new file mode 100644
--- /dev/null
+++ b/Engine/Plugins/Developer/Sandbox/SandboxedEditing/Source/SandboxedEditing/Private/Features/Browser/Views/PersistSandboxView.h
@@ -0,0 +1,29 @@
+﻿// Copyright Epic Games, Inc. All Rights Reserved.
+
+#pragma once
+
+#include "Features/Browser/ViewModels/Persist/PersistSandboxWorkflow.h"
+#include "Templates/SharedPointer.h"
+#include "Templates/UnrealTemplate.h"
+
+namespace UE::SandboxedEditing
+{
+class FPersistSandboxViewModel;
+
+/** Visualizes FPersistSandboxWorkflow. */
+class FPersistSandboxView : public FNoncopyable
+{
+public:
+	
+	explicit FPersistSandboxView(const TSharedRef<FPersistSandboxViewModel>& InPersistViewModel);
+	~FPersistSandboxView();
+	
+private:
+	
+	/** Notifies when FPersistSandboxWorkflow starts. */
+	const TSharedRef<FPersistSandboxViewModel> PersistViewModel;
+	
+	void ShowPersistWorkflow(FPersistSandboxWorkflow& InWorkflow);
+};
+}
+
diff --git a/Engine/Plugins/Developer/Sandbox/SandboxedEditing/Source/SandboxedEditing/Private/Features/Browser/Widgets/Sandboxed/FileState/Columns/ActionFileStateColumn.cpp b/Engine/Plugins/Developer/Sandbox/SandboxedEditing/Source/SandboxedEditing/Private/Features/Browser/Widgets/Sandboxed/FileState/Columns/ActionFileStateColumn.cpp
new file mode 100644
--- /dev/null
+++ b/Engine/Plugins/Developer/Sandbox/SandboxedEditing/Source/SandboxedEditing/Private/Features/Browser/Widgets/Sandboxed/FileState/Columns/ActionFileStateColumn.cpp
@@ -0,0 +1,71 @@
+﻿// Copyright Epic Games, Inc. All Rights Reserved.
+
+#include "ActionFileStateColumn.h"
+
+#include "SandboxedEditingStyle.h"
+#include "Features/Browser/ViewModels/FileState/FileStateColumns.h"
+#include "Features/Browser/ViewModels/FileState/FileStateItem.h"
+#include "RevisionControlStyle/RevisionControlStyle.h"
+#include "Types/SandboxFileChange.h"
+#include "Widgets/Images/SImage.h"
+#include "Widgets/Layout/SBox.h"
+
+#define LOCTEXT_NAMESPACE "FActionFileStateColumn"
+
+namespace UE::SandboxedEditing
+{
+namespace FileActionDetail
+{
+static const FSlateBrush* GetActionBrush(FileSandboxCore::ESandboxFileChange InChange)
+{
+	switch (InChange)
+	{
+	case FileSandboxCore::ESandboxFileChange::Added: return FRevisionControlStyleManager::Get().GetBrush("RevisionControl.OpenForAdd");
+	case FileSandboxCore::ESandboxFileChange::Removed: return FRevisionControlStyleManager::Get().GetBrush("RevisionControl.MarkedForDelete");
+	case FileSandboxCore::ESandboxFileChange::Edited: return FRevisionControlStyleManager::Get().GetBrush("RevisionControl.CheckedOut");
+	case FileSandboxCore::ESandboxFileChange::None: [[fallthrough]];
+	default: return nullptr;
+	}
+}
+
+static FText GetActionToolTipText(FileSandboxCore::ESandboxFileChange InChange)
+{
+	switch (InChange)
+	{
+	case FileSandboxCore::ESandboxFileChange::Added: return LOCTEXT("Added", "Marked for add");
+	case FileSandboxCore::ESandboxFileChange::Removed: return LOCTEXT("Added", "Marked for delete");
+	case FileSandboxCore::ESandboxFileChange::Edited: return LOCTEXT("Added", "File is edited");
+	case FileSandboxCore::ESandboxFileChange::None: [[fallthrough]];
+	default: return FText::GetEmpty();
+	}
+}
+}
+
+SHeaderRow::FColumn::FArguments FActionFileStateColumn::MakeColumnArguments()
+{
+	return SHeaderRow::FColumn::FArguments()
+		.ColumnId(FileActionColumn)
+		.DefaultLabel(FText::GetEmpty())
+		.ToolTipText(LOCTEXT("ToolTip", "The action that was performed on the file"))
+		.FixedWidth(FSandboxedEditingStyle::Get().GetFloat("SandboxedEditing.FileActions.FileActionColumn.FixedWidth"))
+		.ShouldGenerateWidget(true)
+		[
+			SNullWidget::NullWidget
+		];
+}
+
+TSharedRef<SWidget> FActionFileStateColumn::MakeColumnWidget(const FMakeFileStateColumnWidgetArgs& InArgs)
+{
+	return SNew(SBox)
+		.HAlign(HAlign_Center)
+		.VAlign(VAlign_Center)
+		.Padding(FSandboxedEditingStyle::Get().GetMargin("SandboxedEditing.FileActions.FileActionColumn.Padding"))
+		[
+			SNew(SImage)
+			.Image(FileActionDetail::GetActionBrush(InArgs.RowData->Action))
+			.ToolTipText(FileActionDetail::GetActionToolTipText(InArgs.RowData->Action))
+		];
+}
+}
+
+#undef LOCTEXT_NAMESPACE
\ No newline at end of file
diff --git a/Engine/Plugins/Developer/Sandbox/SandboxedEditing/Source/SandboxedEditing/Private/Features/Browser/Widgets/Sandboxed/FileState/Columns/ActionFileStateColumn.h b/Engine/Plugins/Developer/Sandbox/SandboxedEditing/Source/SandboxedEditing/Private/Features/Browser/Widgets/Sandboxed/FileState/Columns/ActionFileStateColumn.h
new file mode 100644
--- /dev/null
+++ b/Engine/Plugins/Developer/Sandbox/SandboxedEditing/Source/SandboxedEditing/Private/Features/Browser/Widgets/Sandboxed/FileState/Columns/ActionFileStateColumn.h
@@ -0,0 +1,20 @@
+﻿// Copyright Epic Games, Inc. All Rights Reserved.
+
+#pragma once
+
+#include "Features/Browser/ViewModels/FileState/IFileStateColumnBehavior.h"
+#include "Features/Browser/ViewModels/FileState/IFileStateColumnWidgetFactory.h"
+
+namespace UE::SandboxedEditing
+{
+class FActionFileStateColumn : public IFileStateColumnBehavior, public IFileStateColumnWidgetFactory
+{
+public:
+	
+	//~ Begin IFileStateWidgetFactory Interface
+	virtual SHeaderRow::FColumn::FArguments MakeColumnArguments() override;
+	virtual TSharedRef<SWidget> MakeColumnWidget(const FMakeFileStateColumnWidgetArgs& InArgs) override;
+	//~ End IFileStateWidgetFactory Interface
+};
+}
+
diff --git a/Engine/Plugins/Developer/Sandbox/SandboxedEditing/Source/SandboxedEditing/Private/Features/Browser/Widgets/Sandboxed/FileState/Columns/PathFileStateColumn.cpp b/Engine/Plugins/Developer/Sandbox/SandboxedEditing/Source/SandboxedEditing/Private/Features/Browser/Widgets/Sandboxed/FileState/Columns/PathFileStateColumn.cpp
new file mode 100644
--- /dev/null
+++ b/Engine/Plugins/Developer/Sandbox/SandboxedEditing/Source/SandboxedEditing/Private/Features/Browser/Widgets/Sandboxed/FileState/Columns/PathFileStateColumn.cpp
@@ -0,0 +1,39 @@
+﻿// Copyright Epic Games, Inc. All Rights Reserved.
+
+#include "PathFileStateColumn.h"
+
+#include "Features/Browser/ViewModels/FileState/FileStateColumns.h"
+#include "Features/Browser/ViewModels/FileState/FileStateItem.h"
+#include "SandboxedEditingStyle.h"
+#include "Widgets/Layout/SBox.h"
+#include "Widgets/Text/STextBlock.h"
+
+#define LOCTEXT_NAMESPACE "FPathFileStateColumn"
+
+namespace UE::SandboxedEditing
+{
+SHeaderRow::FColumn::FArguments FPathFileStateColumn::MakeColumnArguments()
+{
+	return SHeaderRow::FColumn::FArguments()
+		.ColumnId(FilePathColumn)
+		.FillWidth(FSandboxedEditingStyle::Get().GetFloat("SandboxedEditing.FileActions.PathColumn.FillWidth"))
+		.DefaultLabel(LOCTEXT("Name.Label", "Path"))
+		.ToolTipText(LOCTEXT("Name.Description", "The absolute file path to the file"))
+		.ShouldGenerateWidget(true);
+}
+
+TSharedRef<SWidget> FPathFileStateColumn::MakeColumnWidget(const FMakeFileStateColumnWidgetArgs& InArgs)
+{
+	const FText AbsPathAsText = FText::AsCultureInvariant(FPaths::ConvertRelativePathToFull(InArgs.RowData->NonSandboxFile));
+	return SNew(SBox)
+		.VAlign(VAlign_Center)
+		.Padding(FSandboxedEditingStyle::Get().GetMargin("SandboxedEditing.FileActions.PathColumn.Padding"))
+		[
+			SNew(STextBlock)
+			.Text(AbsPathAsText)
+			.OverflowPolicy(ETextOverflowPolicy::MiddleEllipsis)
+		];
+}
+}
+
+#undef LOCTEXT_NAMESPACE
\ No newline at end of file
diff --git a/Engine/Plugins/Developer/Sandbox/SandboxedEditing/Source/SandboxedEditing/Private/Features/Browser/Widgets/Sandboxed/FileState/Columns/PathFileStateColumn.h b/Engine/Plugins/Developer/Sandbox/SandboxedEditing/Source/SandboxedEditing/Private/Features/Browser/Widgets/Sandboxed/FileState/Columns/PathFileStateColumn.h
new file mode 100644
--- /dev/null
+++ b/Engine/Plugins/Developer/Sandbox/SandboxedEditing/Source/SandboxedEditing/Private/Features/Browser/Widgets/Sandboxed/FileState/Columns/PathFileStateColumn.h
@@ -0,0 +1,21 @@
+﻿// Copyright Epic Games, Inc. All Rights Reserved.
+
+#pragma once
+
+#include "Features/Browser/ViewModels/FileState/IFileStateColumnBehavior.h"
+#include "Features/Browser/ViewModels/FileState/IFileStateColumnWidgetFactory.h"
+
+namespace UE::SandboxedEditing
+{
+/** Displays the file path to the file. */
+class FPathFileStateColumn : public IFileStateColumnBehavior, public IFileStateColumnWidgetFactory
+{
+public:
+	
+	//~ Begin IFileStateWidgetFactory Interface
+	virtual SHeaderRow::FColumn::FArguments MakeColumnArguments() override;
+	virtual TSharedRef<SWidget> MakeColumnWidget(const FMakeFileStateColumnWidgetArgs& InArgs) override;
+	//~ End IFileStateWidgetFactory Interface
+};
+}
+
diff --git a/Engine/Plugins/Developer/Sandbox/SandboxedEditing/Source/SandboxedEditing/Private/Features/Browser/Widgets/Sandboxed/FileState/Columns/PersistFileStateColumn.cpp b/Engine/Plugins/Developer/Sandbox/SandboxedEditing/Source/SandboxedEditing/Private/Features/Browser/Widgets/Sandboxed/FileState/Columns/PersistFileStateColumn.cpp
new file mode 100644
--- /dev/null
+++ b/Engine/Plugins/Developer/Sandbox/SandboxedEditing/Source/SandboxedEditing/Private/Features/Browser/Widgets/Sandboxed/FileState/Columns/PersistFileStateColumn.cpp
@@ -0,0 +1,76 @@
+﻿// Copyright Epic Games, Inc. All Rights Reserved.
+
+#include "PersistFileStateColumn.h"
+
+#include "SandboxedEditingStyle.h"
+#include "Features/Browser/ViewModels/FileState/FileStateColumns.h"
+#include "Features/Browser/ViewModels/FileState/FileStateItem.h"
+#include "Features/Browser/ViewModels/Persist/PersistOperationViewModel.h"
+#include "Widgets/Input/SCheckBox.h"
+#include "Widgets/Layout/SBox.h"
+
+#define LOCTEXT_NAMESPACE "FPersistFileStateColumn"
+
+namespace UE::SandboxedEditing
+{
+FPersistFileStateColumn::FPersistFileStateColumn(
+const TSharedRef<FPersistOperationViewModel>& InPersistViewModel
+	)
+	: PersistViewModel(InPersistViewModel)
+{}
+
+SHeaderRow::FColumn::FArguments FPersistFileStateColumn::MakeColumnArguments()
+{
+	return SHeaderRow::FColumn::FArguments()
+		.ColumnId(PersistCheckboxColumn)
+		.FixedWidth(FSandboxedEditingStyle::Get().GetFloat("SandboxedEditing.FileActions.PersistCheckbox.FixedWidth"))
+		.ShouldGenerateWidget(true)
+		[
+			SNew(SCheckBox)
+			.IsChecked_Raw(this, &FPersistFileStateColumn::GetToggleRootSelectedState)
+			.OnCheckStateChanged_Raw(this, &FPersistFileStateColumn::OnToggleRootSelectedCheckBox)
+		];
+}
+
+TSharedRef<SWidget> FPersistFileStateColumn::MakeColumnWidget(const FMakeFileStateColumnWidgetArgs& InArgs)
+{
+	const int32 Index = PersistViewModel->GetPersistableFiles().NonSandboxPaths.IndexOfByKey(InArgs.RowData->NonSandboxFile);
+	if (Index == INDEX_NONE)
+	{
+		return SNullWidget::NullWidget;
+	}
+	
+	return SNew(SBox)
+		.HAlign(HAlign_Center)
+		.Padding(FSandboxedEditingStyle::Get().GetMargin("SandboxedEditing.FileActions.PersistCheckbox.Padding"))
+		[
+			SNew(SCheckBox)
+			.IsChecked_Raw(this, &FPersistFileStateColumn::GetItemCheckBoxState, Index)
+			.OnCheckStateChanged_Raw(this, &FPersistFileStateColumn::SetItemCheckBoxState, Index)
+		];
+}
+
+ECheckBoxState FPersistFileStateColumn::GetToggleRootSelectedState() const
+{
+	return PersistViewModel->AreAllFilesMarkedForPersist() ? ECheckBoxState::Checked : ECheckBoxState::Unchecked;
+}
+
+void FPersistFileStateColumn::OnToggleRootSelectedCheckBox(ECheckBoxState InNewState) const
+{
+	PersistViewModel->SetAllFilesPersisted(InNewState == ECheckBoxState::Checked);
+}
+
+ECheckBoxState FPersistFileStateColumn::GetItemCheckBoxState(int32 InIndex) const
+{
+	return PersistViewModel->IsFileMarkedForPersist(InIndex) 
+		? ECheckBoxState::Checked 
+		: ECheckBoxState::Unchecked;
+}
+
+void FPersistFileStateColumn::SetItemCheckBoxState(ECheckBoxState InNewState, int32 InIndex) const
+{
+	PersistViewModel->SetFilePersisted(InIndex, InNewState == ECheckBoxState::Checked);
+}
+}
+
+#undef LOCTEXT_NAMESPACE
\ No newline at end of file
diff --git a/Engine/Plugins/Developer/Sandbox/SandboxedEditing/Source/SandboxedEditing/Private/Features/Browser/Widgets/Sandboxed/FileState/Columns/PersistFileStateColumn.h b/Engine/Plugins/Developer/Sandbox/SandboxedEditing/Source/SandboxedEditing/Private/Features/Browser/Widgets/Sandboxed/FileState/Columns/PersistFileStateColumn.h
new file mode 100644
--- /dev/null
+++ b/Engine/Plugins/Developer/Sandbox/SandboxedEditing/Source/SandboxedEditing/Private/Features/Browser/Widgets/Sandboxed/FileState/Columns/PersistFileStateColumn.h
@@ -0,0 +1,35 @@
+﻿// Copyright Epic Games, Inc. All Rights Reserved.
+
+#pragma once
+
+#include "Features/Browser/ViewModels/FileState/IFileStateColumnBehavior.h"
+#include "Features/Browser/ViewModels/FileState/IFileStateColumnWidgetFactory.h"
+
+namespace UE::SandboxedEditing
+{
+class FPersistOperationViewModel;
+
+class FPersistFileStateColumn : public IFileStateColumnBehavior, public IFileStateColumnWidgetFactory
+{
+public:
+	
+	explicit FPersistFileStateColumn(const TSharedRef<FPersistOperationViewModel>& InPersistViewModel);
+
+	//~ Begin IFileStateWidgetFactory Interface
+	virtual SHeaderRow::FColumn::FArguments MakeColumnArguments() override;
+	virtual TSharedRef<SWidget> MakeColumnWidget(const FMakeFileStateColumnWidgetArgs& InArgs) override;
+	//~ End IFileStateWidgetFactory Interface
+	
+private:
+	
+	/** Used to set whether an item should be persisted */
+	const TSharedRef<FPersistOperationViewModel> PersistViewModel;
+
+	ECheckBoxState GetToggleRootSelectedState() const;
+	void OnToggleRootSelectedCheckBox(ECheckBoxState InNewState) const;
+
+	ECheckBoxState GetItemCheckBoxState(int32 InIndex) const;
+	void SetItemCheckBoxState(ECheckBoxState InNewState, int32 InIndex) const;
+};
+}
+
diff --git a/Engine/Plugins/Developer/Sandbox/SandboxedEditing/Source/SandboxedEditing/Private/Features/Browser/Widgets/Sandboxed/FileState/SFileStateListRow.cpp b/Engine/Plugins/Developer/Sandbox/SandboxedEditing/Source/SandboxedEditing/Private/Features/Browser/Widgets/Sandboxed/FileState/SFileStateListRow.cpp
new file mode 100644
--- /dev/null
+++ b/Engine/Plugins/Developer/Sandbox/SandboxedEditing/Source/SandboxedEditing/Private/Features/Browser/Widgets/Sandboxed/FileState/SFileStateListRow.cpp
@@ -0,0 +1,32 @@
+﻿// Copyright Epic Games, Inc. All Rights Reserved.
+
+#include "SFileStateListRow.h"
+
+#include "Features/Browser/ViewModels/FileState/IFileStateColumnWidgetFactory.h"
+
+namespace UE::SandboxedEditing
+{
+void SFileStateListRow::Construct(const FArguments& InArgs, const TSharedRef<FFileStateItem>& InItem, const TSharedRef<STableViewBase>& InOwnerTable)
+{
+	Item = InItem;
+	ColumnFactoryMap = InArgs._ColumnFactories;
+	
+	SMultiColumnTableRow::Construct(
+		SMultiColumnTableRow::FArguments().Padding(InArgs._Padding),
+		InOwnerTable
+		);
+}
+
+TSharedRef<SWidget> SFileStateListRow::GenerateWidgetForColumn(const FName& InColumnName)
+{
+	if (TSharedRef<IFileStateColumnWidgetFactory>* Factory = ColumnFactoryMap.Find(InColumnName))
+	{
+		return Factory->Get().MakeColumnWidget(
+			FMakeFileStateColumnWidgetArgs(Item)
+		);
+	}
+	
+	checkNoEntry();
+	return SNullWidget::NullWidget;
+}
+}
diff --git a/Engine/Plugins/Developer/Sandbox/SandboxedEditing/Source/SandboxedEditing/Private/Features/Browser/Widgets/Sandboxed/FileState/SFileStateListRow.h b/Engine/Plugins/Developer/Sandbox/SandboxedEditing/Source/SandboxedEditing/Private/Features/Browser/Widgets/Sandboxed/FileState/SFileStateListRow.h
new file mode 100644
--- /dev/null
+++ b/Engine/Plugins/Developer/Sandbox/SandboxedEditing/Source/SandboxedEditing/Private/Features/Browser/Widgets/Sandboxed/FileState/SFileStateListRow.h
@@ -0,0 +1,41 @@
+﻿// Copyright Epic Games, Inc. All Rights Reserved.
+
+#pragma once
+
+#include "Features/Browser/ViewModels/FileState/FileStateColumns.h"
+#include "Features/Browser/ViewModels/FileState/FileStateItem.h"
+#include "Features/Browser/ViewModels/FileState/IFileStateColumnWidgetFactory.h"
+#include "Features/Browser/ViewModels/List/ISandboxColumnWidgetFactory.h"
+#include "Widgets/DeclarativeSyntaxSupport.h"
+#include "Widgets/Views/SListView.h"
+#include "Widgets/Views/STableRow.h"
+
+namespace UE::SandboxedEditing
+{
+class SFileStateListRow : public SMultiColumnTableRow<TSharedPtr<FFileStateItem>>
+{
+public:
+	
+	SLATE_BEGIN_ARGS(SFileStateListRow){}
+		/** Used to build the columns. */
+		SLATE_ARGUMENT(FFileStateColumnFactoryMap, ColumnFactories)
+		/** Padding for the row */
+		SLATE_ARGUMENT(FMargin, Padding)
+	SLATE_END_ARGS()
+	
+	void Construct(const FArguments& InArgs, const TSharedRef<FFileStateItem>& InItem, const TSharedRef<STableViewBase>& InOwnerTable);
+
+	//~ Begin SMultiColumnTableRow Interface
+	virtual TSharedRef<SWidget> GenerateWidgetForColumn(const FName& InColumnName) override;
+	//~ End SMultiColumnTableRow Interface
+	
+private:
+	
+	/** The item that is being displayed by this row. */
+	TSharedPtr<FFileStateItem> Item;
+	
+	/** Used to build the columns. */
+	FFileStateColumnFactoryMap ColumnFactoryMap;
+};
+}
+
diff --git a/Engine/Plugins/Developer/Sandbox/SandboxedEditing/Source/SandboxedEditing/Private/Features/Browser/Widgets/Sandboxed/FileState/SFileStateListView.cpp b/Engine/Plugins/Developer/Sandbox/SandboxedEditing/Source/SandboxedEditing/Private/Features/Browser/Widgets/Sandboxed/FileState/SFileStateListView.cpp
new file mode 100644
--- /dev/null
+++ b/Engine/Plugins/Developer/Sandbox/SandboxedEditing/Source/SandboxedEditing/Private/Features/Browser/Widgets/Sandboxed/FileState/SFileStateListView.cpp
@@ -0,0 +1,52 @@
+﻿// Copyright Epic Games, Inc. All Rights Reserved.
+
+#include "SFileStateListView.h"
+
+#include "SandboxedEditingStyle.h"
+#include "SFileStateListRow.h"
+#include "Features/Browser/ViewModels/FileState/IFileStateColumnWidgetFactory.h"
+#include "Features/Browser/ViewModels/FileState/IFileStateViewModel.h"
+#include "Widgets/Views/SHeaderRow.h"
+
+namespace UE::SandboxedEditing
+{
+void SFileStateListView::Construct(const FArguments& InArgs, const TSharedRef<IFileStateViewModel>& InViewModel)
+{
+	ViewModel = InViewModel;
+	ColumnFactories = InArgs._ColumnFactories;
+	
+	ChildSlot
+	[
+		SAssignNew(ListView, SListView<TSharedPtr<FFileStateItem>>)
+		.ListItemsSource(&ViewModel->GetItems())
+		.OnGenerateRow(this, &SFileStateListView::OnGenerateRow)
+		.SelectionMode(ESelectionMode::None)
+		.HeaderRow(MakeHeaderRow())
+	];
+}
+
+TSharedRef<SHeaderRow> SFileStateListView::MakeHeaderRow()
+{
+	const TSharedRef<SHeaderRow> HeaderRow = SNew(SHeaderRow);
+		
+	for (FName ColumnId : ViewModel->GetDisplayedColumns())
+	{
+		if (const TSharedRef<IFileStateColumnWidgetFactory>* Factory = ColumnFactories.Find(ColumnId))
+		{
+			HeaderRow->AddColumn(Factory->Get().MakeColumnArguments());
+		}
+	}
+	
+	return HeaderRow;
+}
+
+TSharedRef<ITableRow> SFileStateListView::OnGenerateRow(TSharedPtr<FFileStateItem> InItem, const TSharedRef<STableViewBase>& InOwnerTable) const
+{
+	check(InItem);
+	const FMargin Padding = FSandboxedEditingStyle::Get().GetMargin("SandboxedEditing.FileActions.RowPadding");
+	
+	return SNew(SFileStateListRow, InItem.ToSharedRef(), InOwnerTable)
+		.ColumnFactories(ColumnFactories)
+		.Padding(Padding);
+}
+}
diff --git a/Engine/Plugins/Developer/Sandbox/SandboxedEditing/Source/SandboxedEditing/Private/Features/Browser/Widgets/Sandboxed/FileState/SFileStateListView.h b/Engine/Plugins/Developer/Sandbox/SandboxedEditing/Source/SandboxedEditing/Private/Features/Browser/Widgets/Sandboxed/FileState/SFileStateListView.h
new file mode 100644
--- /dev/null
+++ b/Engine/Plugins/Developer/Sandbox/SandboxedEditing/Source/SandboxedEditing/Private/Features/Browser/Widgets/Sandboxed/FileState/SFileStateListView.h
@@ -0,0 +1,48 @@
+﻿// Copyright Epic Games, Inc. All Rights Reserved.
+
+#pragma once
+
+#include "Features/Browser/ViewModels/FileState/FileStateColumns.h"
+#include "Widgets/DeclarativeSyntaxSupport.h"
+#include "Widgets/SCompoundWidget.h"
+#include "Widgets/Views/SListView.h"
+
+class STableViewBase;
+class ITableRow;
+class SHeaderRow;
+
+namespace UE::SandboxedEditing
+{
+class IFileStateViewModel;
+
+/** Displays a list view FFileStateItem provided by IFileStateViewModel. */
+class SFileStateListView : public SCompoundWidget
+{
+public:
+	
+	SLATE_BEGIN_ARGS(SFileStateListView) {}
+		/** Used to build the columns. */
+		SLATE_ARGUMENT(FFileStateColumnFactoryMap, ColumnFactories)
+	SLATE_END_ARGS()
+	
+	void Construct(const FArguments& InArgs, const TSharedRef<IFileStateViewModel>& InViewModel);
+	
+private:
+	
+	/** The view model being displayed */
+	TSharedPtr<IFileStateViewModel> ViewModel;
+	
+	/** The list view displaying the items. */
+	TSharedPtr<SListView<TSharedPtr<FFileStateItem>>> ListView;
+	
+	/** Used to build the columns. */
+	TMap<FName, TSharedRef<IFileStateColumnWidgetFactory>> ColumnFactories;
+	
+	/** @return The header row to display at the top of the table */
+	TSharedRef<SHeaderRow> MakeHeaderRow();
+	
+	/** @return Widget for the items */
+	TSharedRef<ITableRow> OnGenerateRow(TSharedPtr<FFileStateItem> InItem, const TSharedRef<STableViewBase>& InOwnerTable) const;
+};
+}
+
diff --git a/Engine/Plugins/Developer/Sandbox/SandboxedEditing/Source/SandboxedEditing/Private/Features/Browser/Widgets/Sandboxed/SPersistOperationWidget.cpp b/Engine/Plugins/Developer/Sandbox/SandboxedEditing/Source/SandboxedEditing/Private/Features/Browser/Widgets/Sandboxed/SPersistOperationWidget.cpp
new file mode 100644
--- /dev/null
+++ b/Engine/Plugins/Developer/Sandbox/SandboxedEditing/Source/SandboxedEditing/Private/Features/Browser/Widgets/Sandboxed/SPersistOperationWidget.cpp
@@ -0,0 +1,147 @@
+﻿// Copyright Epic Games, Inc. All Rights Reserved.
+
+#include "SPersistOperationWidget.h"
+
+#include "SPositiveActionButton.h"
+#include "Features/Browser/ViewModels/FileState/FileStateItem.h"
+#include "Features/Browser/ViewModels/FileState/StaticFileStateViewModel.h"
+#include "Features/Browser/ViewModels/Persist/PersistOperationViewModel.h"
+#include "FileState/SFileStateListView.h"
+#include "FileState/Columns/ActionFileStateColumn.h"
+#include "FileState/Columns/PathFileStateColumn.h"
+#include "FileState/Columns/PersistFileStateColumn.h"
+#include "Styling/AppStyle.h"
+#include "Types/GatheredFileChanges.h"
+#include "Widgets/Input/SButton.h"
+#include "Widgets/Layout/SBorder.h"
+#include "Widgets/Layout/SUniformGridPanel.h"
+#include "Widgets/SBoxPanel.h"
+
+#define LOCTEXT_NAMESPACE "SPersistWidget"
+
+namespace UE::SandboxedEditing
+{
+namespace PersistOperationDetail
+{
+static TArray<TSharedPtr<FFileStateItem>> MakeItemsFromFileChanges(const FileSandboxCore::FGatheredFileChanges& InChanges)
+{
+	const TArray<FString>& NonSandboxPaths = InChanges.NonSandboxPaths;
+	const TArray<FileSandboxCore::ESandboxFileChange>& FileActions = InChanges.FileActions;
+	check(NonSandboxPaths.Num() == FileActions.Num());
+	
+	TArray<TSharedPtr<FFileStateItem>> Items;
+	for (int32 Index = 0; Index < NonSandboxPaths.Num(); ++Index)
+	{
+		Items.Emplace(MakeShared<FFileStateItem>(NonSandboxPaths[Index], FileActions[Index]));
+	}
+	return Items;
+}
+
+static void MakePersistColumns(
+	const TSharedRef<FPersistOperationViewModel>& InViewModel,
+	TMap<FName, TSharedRef<IFileStateColumnBehavior>>& OutBehaviors, 
+	TMap<FName, TSharedRef<IFileStateColumnWidgetFactory>>& OutFactories
+	)
+{
+	const TSharedRef<FPersistFileStateColumn> CheckboxColumn = MakeShared<FPersistFileStateColumn>(InViewModel);
+	const TSharedRef<FActionFileStateColumn> ActionColum = MakeShared<FActionFileStateColumn>();
+	const TSharedRef<FPathFileStateColumn> PathColumn = MakeShared<FPathFileStateColumn>();
+	
+	OutBehaviors = TMap<FName, TSharedRef<IFileStateColumnBehavior>>
+	{
+		{ PersistCheckboxColumn, CheckboxColumn },
+		{ FileActionColumn, ActionColum },
+		{ FilePathColumn, PathColumn }
+	};
+	OutFactories = TMap<FName, TSharedRef<IFileStateColumnWidgetFactory>>
+	{
+		{ PersistCheckboxColumn, CheckboxColumn },
+		{ FileActionColumn, ActionColum },
+		{ FilePathColumn, PathColumn }
+	};
+}
+}
+
+void SPersistOperationWidget::Construct(const FArguments& InArgs, const TSharedRef<FPersistOperationViewModel>& InViewModel)
+{
+	ViewModel = InViewModel;
+	
+	TMap<FName, TSharedRef<IFileStateColumnBehavior>> ColumnBehaviors;
+	TMap<FName, TSharedRef<IFileStateColumnWidgetFactory>> ColumnFactories;
+	PersistOperationDetail::MakePersistColumns(ViewModel.ToSharedRef(), ColumnBehaviors, ColumnFactories);
+	const TSharedRef<FStaticFileStateViewModel> FileStateViewModel = MakeShared<FStaticFileStateViewModel>(
+		ColumnBehaviors, 
+		PersistOperationDetail::MakeItemsFromFileChanges(ViewModel->GetPersistableFiles())
+		);
+	
+	ChildSlot
+	[
+		SNew(SBorder)
+		.BorderImage(FAppStyle::GetBrush("ToolPanel.GroupBorder"))
+		[
+			SNew(SVerticalBox)
+
+			+SVerticalBox::Slot()
+			.Padding(5.f)
+			[
+				SNew(SBorder)
+				[
+					SNew(SFileStateListView, FileStateViewModel)
+					.ColumnFactories(ColumnFactories)
+				]
+			]
+
+			+SVerticalBox::Slot()
+			.AutoHeight()
+			.HAlign(HAlign_Right)
+			.VAlign(VAlign_Bottom)
+			.Padding(5.f)
+			[
+				MakeButtonArea()
+			]
+		]
+	];
+}
+
+TSharedRef<SWidget> SPersistOperationWidget::MakeButtonArea()
+{
+	return SNew(SUniformGridPanel)
+		.SlotPadding(FAppStyle::GetMargin("StandardDialog.SlotPadding"))
+		.MinDesiredSlotWidth(FAppStyle::GetFloat("StandardDialog.MinDesiredSlotWidth"))
+		.MinDesiredSlotHeight(FAppStyle::GetFloat("StandardDialog.MinDesiredSlotHeight"))
+		+SUniformGridPanel::Slot(0,0)
+		[
+			SNew(SPositiveActionButton)
+			.Text(LOCTEXT("PersistButton", "Persist"))
+			.IsEnabled(this, &SPersistOperationWidget::IsPersistButtonEnabled)
+			.OnClicked(this, &SPersistOperationWidget::OnClickPersist)
+		]
+		+SUniformGridPanel::Slot(1,0)
+		[
+			SNew(SButton)
+			.HAlign(HAlign_Center)
+			.ContentPadding(FAppStyle::GetMargin("StandardDialog.ContentPadding"))
+			.Text(LOCTEXT("CancelButton", "Cancel"))
+			.OnClicked(this, &SPersistOperationWidget::OnClickCancel)
+		];
+}
+
+bool SPersistOperationWidget::IsPersistButtonEnabled() const
+{
+	return ViewModel->AreAnyFilesMarkedForPersist();
+}
+
+FReply SPersistOperationWidget::OnClickPersist() const
+{
+	ViewModel->ConfirmPersist();
+	return FReply::Handled();
+}
+
+FReply SPersistOperationWidget::OnClickCancel() const
+{
+	ViewModel->CancelPersist();
+	return FReply::Handled();
+}
+}
+
+#undef LOCTEXT_NAMESPACE
\ No newline at end of file
diff --git a/Engine/Plugins/Developer/Sandbox/SandboxedEditing/Source/SandboxedEditing/Private/Features/Browser/Widgets/Sandboxed/SPersistOperationWidget.h b/Engine/Plugins/Developer/Sandbox/SandboxedEditing/Source/SandboxedEditing/Private/Features/Browser/Widgets/Sandboxed/SPersistOperationWidget.h
new file mode 100644
--- /dev/null
+++ b/Engine/Plugins/Developer/Sandbox/SandboxedEditing/Source/SandboxedEditing/Private/Features/Browser/Widgets/Sandboxed/SPersistOperationWidget.h
@@ -0,0 +1,35 @@
+﻿// Copyright Epic Games, Inc. All Rights Reserved.
+
+#pragma once
+
+#include "Widgets/DeclarativeSyntaxSupport.h"
+#include "Widgets/SCompoundWidget.h"
+
+namespace UE::SandboxedEditing
+{
+class FStaticFileStateViewModel;
+class FPersistOperationViewModel;
+
+/** Displays all UI for persisting a sandbox */
+class SPersistOperationWidget : public SCompoundWidget
+{
+public:
+	
+	SLATE_BEGIN_ARGS(SPersistOperationWidget){}
+	SLATE_END_ARGS()
+	
+	void Construct(const FArguments& InArgs, const TSharedRef<FPersistOperationViewModel>& InViewModel);
+	
+private:
+	
+	/** Gets the view model being visualized. */
+	TSharedPtr<FPersistOperationViewModel> ViewModel;
+
+	bool IsPersistButtonEnabled() const;
+	TSharedRef<SWidget> MakeButtonArea();
+
+	FReply OnClickPersist() const;
+	FReply OnClickCancel() const;
+};
+}
+
diff --git a/Engine/Plugins/Developer/Sandbox/SandboxedEditing/Source/SandboxedEditing/Private/Framework/Models/SandboxSystemModel.cpp b/Engine/Plugins/Developer/Sandbox/SandboxedEditing/Source/SandboxedEditing/Private/Framework/Models/SandboxSystemModel.cpp
--- a/Engine/Plugins/Developer/Sandbox/SandboxedEditing/Source/SandboxedEditing/Private/Framework/Models/SandboxSystemModel.cpp
+++ b/Engine/Plugins/Developer/Sandbox/SandboxedEditing/Source/SandboxedEditing/Private/Framework/Models/SandboxSystemModel.cpp
@@ -159,7 +159,14 @@ bool FSandboxSystemModel::IsAllowedToLeaveSandbox(FText* OutReason) const
 
 bool FSandboxSystemModel::CanLeaveSandboxWithoutFurtherActions() const
 {
-	return IsAllowedToLeaveSandbox() && GetInMemoryChanges().IsEmpty();
+	using namespace FileSandboxCore;
+	IFileSandboxCoreModule& Module = IFileSandboxCoreModule::Get();
+	ISandboxManager& SandboxManager = Module.GetSandboxManager();
+	ISandboxInstance* Instance = SandboxManager.GetActiveSandboxInstance();
+	
+	const bool bHasSandboxedChanges = Instance && !Instance->HasFileChanges();
+	const bool bHasInMemoryChanges = GetInMemoryChanges().IsEmpty();
+	return IsAllowedToLeaveSandbox() && bHasSandboxedChanges && bHasInMemoryChanges;
 }
 
 TArray<UPackage*> FSandboxSystemModel::GetInMemoryChanges() const
@@ -222,11 +229,31 @@ void FSandboxSystemModel::PersistAllChanges() const
 	}
 }
 
+void FSandboxSystemModel::PersistFiles(TConstArrayView<FString> InFiles) const
+{
+	using namespace FileSandboxCore;
+	IFileSandboxCoreModule& Module = IFileSandboxCoreModule::Get();
+	ISandboxManager& SandboxManager = Module.GetSandboxManager();
+	if (ISandboxInstance* Instance = SandboxManager.GetActiveSandboxInstance())
+	{
+		Instance->PersistSandbox(FPersistArgs{ InFiles });
+	}
+}
+
 bool FSandboxSystemModel::CanPersistAllChanges() const
 {
 	return HasActiveSandbox();
 }
 
+FileSandboxCore::FGatheredFileChanges FSandboxSystemModel::GatherFileChanges() const
+{
+	using namespace FileSandboxCore;
+	IFileSandboxCoreModule& Module = IFileSandboxCoreModule::Get();
+	ISandboxManager& SandboxManager = Module.GetSandboxManager();
+	ISandboxInstance* Instance = SandboxManager.GetActiveSandboxInstance();
+	return Instance ? Instance->GatherChangedFiles() : FGatheredFileChanges{};
+}
+
 void FSandboxSystemModel::DiscardAllChanges() const
 {
 	using namespace FileSandboxCore;
diff --git a/Engine/Plugins/Developer/Sandbox/SandboxedEditing/Source/SandboxedEditing/Private/Framework/Models/SandboxSystemModel.h b/Engine/Plugins/Developer/Sandbox/SandboxedEditing/Source/SandboxedEditing/Private/Framework/Models/SandboxSystemModel.h
--- a/Engine/Plugins/Developer/Sandbox/SandboxedEditing/Source/SandboxedEditing/Private/Framework/Models/SandboxSystemModel.h
+++ b/Engine/Plugins/Developer/Sandbox/SandboxedEditing/Source/SandboxedEditing/Private/Framework/Models/SandboxSystemModel.h
@@ -67,8 +67,13 @@ public:
 
 	/** Persists all changes */
 	void PersistAllChanges() const;
+	/** Persists given files. */
+	void PersistFiles(TConstArrayView<FString> InFiles) const;
 	/** @return Whether all changes can be persisted */
 	bool CanPersistAllChanges() const;
+	
+	/** @return The file changes that were performed by the sandbox. */
+	FileSandboxCore::FGatheredFileChanges GatherFileChanges() const;
 
 	/** Discards all changes */
 	void DiscardAllChanges() const;
diff --git a/Engine/Plugins/Developer/Sandbox/SandboxedEditing/Source/SandboxedEditing/Private/SandboxedEditingStyle.cpp b/Engine/Plugins/Developer/Sandbox/SandboxedEditing/Source/SandboxedEditing/Private/SandboxedEditingStyle.cpp
--- a/Engine/Plugins/Developer/Sandbox/SandboxedEditing/Source/SandboxedEditing/Private/SandboxedEditingStyle.cpp
+++ b/Engine/Plugins/Developer/Sandbox/SandboxedEditing/Source/SandboxedEditing/Private/SandboxedEditingStyle.cpp
@@ -42,6 +42,7 @@ void FSandboxedEditingStyle::RegisterBrowserStyle()
 	Set("SandboxedEditing.Browser.DeleteSandbox", new IMAGE_BRUSH("DeleteSandbox_48", Icon24x24, IconColorAndOpacity));
 	Set("SandboxedEditing.Browser.LeaveSandbox", new IMAGE_BRUSH("LeaveSandbox_48", Icon24x24, IconColorAndOpacity));
 	
+	// Sandbox browser
 	Set("SandboxedEditing.Browser.RowPadding", FMargin(0.f, 4.f));
 	Set("SandboxedEditing.Browser.NameColumn.Padding", FMargin(16.f, 0.f, 0.f, 0.f));
 	Set("SandboxedEditing.Browser.NameColumn.FillWidth", 2.f);
@@ -50,6 +51,13 @@ void FSandboxedEditingStyle::RegisterBrowserStyle()
 	Set("SandboxedEditing.Browser.LastModifiedColumn.Padding", FMargin(8.f, 0.f, 0.f, 0.f));
 	Set("SandboxedEditing.Browser.LastModifiedColumn.FillWidth", 1.f);
 	
+	Set("SandboxedEditing.FileActions.RowPadding", FMargin(0.f, 4.f));
+	Set("SandboxedEditing.FileActions.PathColumn.Padding", FMargin(16.f, 0.f, 0.f, 0.f));
+	Set("SandboxedEditing.FileActions.PathColumn.FillWidth", 2.f);
+	Set("SandboxedEditing.FileActions.PersistCheckbox.Padding", FMargin(2.f, 0.f, 0.f, 0.f));
+	Set("SandboxedEditing.FileActions.PersistCheckbox.FixedWidth", 23.f);
+	Set("SandboxedEditing.FileActions.FileActionColumn.Padding", FMargin(0.f, 0.f, 0.f, 0.f));
+	Set("SandboxedEditing.FileActions.FileActionColumn.FixedWidth", 23.f);
 }
 
 void FSandboxedEditingStyle::RegisterSharedStyle()
diff --git a/Engine/Plugins/Developer/Sandbox/SandboxedEditing/Source/SandboxedEditing/SandboxedEditing.Build.cs b/Engine/Plugins/Developer/Sandbox/SandboxedEditing/Source/SandboxedEditing/SandboxedEditing.Build.cs
--- a/Engine/Plugins/Developer/Sandbox/SandboxedEditing/Source/SandboxedEditing/SandboxedEditing.Build.cs
+++ b/Engine/Plugins/Developer/Sandbox/SandboxedEditing/Source/SandboxedEditing/SandboxedEditing.Build.cs
@@ -23,6 +23,7 @@ namespace UnrealBuildTool.Rules
 					"InputCore",
 					"Slate", 
 					"SlateCore", 
+					"SourceControl",
 					"ToolMenus", 
 					"ToolWidgets", 
 					"WorkspaceMenuStructure", 
