Sandboxed Editing: When persisting files,
- show a progress bar,
- display a notification after the operation completes

#Synced-CL 49159615

diff --git a/Engine/Plugins/Developer/Sandbox/FileSandbox/Source/FileSandboxUI/Private/Persist/Feedback/SummaryPersistFeedback.cpp b/Engine/Plugins/Developer/Sandbox/FileSandbox/Source/FileSandboxUI/Private/Persist/Feedback/SummaryPersistFeedback.cpp
new file mode 100644
--- /dev/null
+++ b/Engine/Plugins/Developer/Sandbox/FileSandbox/Source/FileSandboxUI/Private/Persist/Feedback/SummaryPersistFeedback.cpp
@@ -0,0 +1,64 @@
+﻿// Copyright Epic Games, Inc. All Rights Reserved.
+
+#include "Persist/Feedback/SummaryPersistFeedback.h"
+
+#include "Misc/Paths.h"
+#include "Types/SandboxFileChange.h"
+
+namespace UE::FileSandboxUI
+{
+FSummaryPersistFeedback::FSummaryPersistFeedback(
+	TConstArrayView<FString> InPersistedFiles,
+	TConstArrayView<FileSandboxCore::ESandboxFileChange> InFileActions
+	)
+{
+	check(InPersistedFiles.Num() == InFileActions.Num());
+	for (int32 Index = 0; Index < InPersistedFiles.Num(); ++Index)
+	{
+		const FileSandboxCore::ESandboxFileChange Change = InFileActions[Index];
+		FileActions.Add(FPaths::ConvertRelativePathToFull(InPersistedFiles[Index]), Change);
+		
+		// We'll decrement it again if an error occurs. We assume errors should be rare.
+		switch (Change)
+		{
+		case FileSandboxCore::ESandboxFileChange::Added: 
+			++Summary.Added;
+			break;
+		case FileSandboxCore::ESandboxFileChange::Removed:
+			++Summary.Deleted;
+			break;
+		case FileSandboxCore::ESandboxFileChange::Edited:
+			++Summary.Edited;
+			break;
+		case FileSandboxCore::ESandboxFileChange::None: [[fallthrough]];
+		default: ensure(false);
+		}
+	}
+}
+
+void FSummaryPersistFeedback::HandleError(const TCHAR* InToFilename)
+{
+	FString Abs = FPaths::ConvertRelativePathToFull(InToFilename);
+	
+	const FileSandboxCore::ESandboxFileChange* FileChange = FileActions.Find(Abs);
+	if (ensure(FileChange))
+	{
+		switch (*FileChange)
+		{
+		case FileSandboxCore::ESandboxFileChange::Added: 
+			--Summary.Added;
+			break;
+		case FileSandboxCore::ESandboxFileChange::Removed:
+			--Summary.Deleted;
+			break;
+		case FileSandboxCore::ESandboxFileChange::Edited:
+			--Summary.Edited;
+			break;
+		case FileSandboxCore::ESandboxFileChange::None: [[fallthrough]];
+		default: ensure(false);
+		}
+	}
+	
+	Summary.FailedFiles.Emplace(MoveTemp(Abs));
+}
+}
diff --git a/Engine/Plugins/Developer/Sandbox/FileSandbox/Source/FileSandboxUI/Public/Persist/Feedback/PersistSummary.h b/Engine/Plugins/Developer/Sandbox/FileSandbox/Source/FileSandboxUI/Public/Persist/Feedback/PersistSummary.h
new file mode 100644
--- /dev/null
+++ b/Engine/Plugins/Developer/Sandbox/FileSandbox/Source/FileSandboxUI/Public/Persist/Feedback/PersistSummary.h
@@ -0,0 +1,23 @@
+﻿// Copyright Epic Games, Inc. All Rights Reserved.
+
+#pragma once
+
+#include "Containers/UnrealString.h"
+#include "HAL/Platform.h"
+
+namespace UE::FileSandboxUI
+{
+/** Summary of a persist operation. */
+struct FPersistSummary
+{
+	/** Number of edited files. */
+	int32 Edited = 0;
+	/** Number of added files. */
+	int32 Added = 0;
+	/** Number of deletes files. */
+	int32 Deleted = 0;
+	
+	/** Non-sandbox paths of files that were not persisted. */
+	TArray<FString> FailedFiles;
+};
+}
\ No newline at end of file
diff --git a/Engine/Plugins/Developer/Sandbox/FileSandbox/Source/FileSandboxUI/Public/Persist/Feedback/SlowTaskPersistFeedback.h b/Engine/Plugins/Developer/Sandbox/FileSandbox/Source/FileSandboxUI/Public/Persist/Feedback/SlowTaskPersistFeedback.h
new file mode 100644
--- /dev/null
+++ b/Engine/Plugins/Developer/Sandbox/FileSandbox/Source/FileSandboxUI/Public/Persist/Feedback/SlowTaskPersistFeedback.h
@@ -0,0 +1,55 @@
+﻿// Copyright Epic Games, Inc. All Rights Reserved.
+
+#pragma once
+
+#include "Interface/IPersistFeedback.h"
+#include "Internationalization/Internationalization.h"
+#include "Misc/Paths.h"
+#include "Misc/ScopedSlowTask.h"
+
+namespace UE::FileSandboxUI
+{
+DECLARE_DELEGATE_RetVal_OneParam(FText, FFormatPathDelegate, const TCHAR* InFileName);
+
+/**
+ * Starts a slow task for the duration of the persist operation. 
+ * Remember that nested of FScopedSlowTask works.
+ */
+class FSlowTaskPersistFeedback : public FileSandboxCore::IPersistFeedback
+{
+public:
+	
+	/** Spawns the progress bar. */
+	FScopedSlowTask SlowTask;
+	
+	/** Optional delegate that produces the progress bar text. */
+	FFormatPathDelegate FormatPathDelegate; 
+	
+	explicit FSlowTaskPersistFeedback(int32 InNumFilesToPersist, FFormatPathDelegate InFormatPath = FFormatPathDelegate())
+		: SlowTask(InNumFilesToPersist)
+		, FormatPathDelegate(InFormatPath)
+	{}
+
+	//~ Begin IPersistFeedback Interface
+	virtual void StartFile(const TCHAR* InFilename) override
+	{
+		const FText Text = FormatPathDelegate.IsBound() 
+			? FormatPathDelegate.Execute(InFilename)
+			: FText::Format(
+				NSLOCTEXT("FProgressBarPersistFeedback", "PersistFileFmt", "Persisting {0}"),
+				FText::AsCultureInvariant(FPaths::GetBaseFilename(InFilename))
+				);
+		SlowTask.EnterProgressFrame(1, Text);
+	}
+	virtual void HandleSuccess(const TCHAR* InFilename) override {}
+	virtual void HandleError_CheckoutNotAllowed(const TCHAR* InFilename) override {}
+	virtual void HandleError_Checkout(const TCHAR* InFilename) override {}
+	virtual void HandleError_Revert(const TCHAR* InFilename) override {}
+	virtual void HandleError_MarkForAdd(const TCHAR* InFilename) override {}
+	virtual void HandleError_DeleteSCC(const TCHAR* InFilename) override {}
+	virtual void HandleError_MakeWritable(const TCHAR* InFilename) override {}
+	virtual void HandleError_MoveFile(const TCHAR* InToFilename, const TCHAR* InFromFilename) override {}
+	virtual void HandleError_DeleteFile(const TCHAR* InFilename) override {}
+	//~ End IPersistFeedback Interface
+};
+}
diff --git a/Engine/Plugins/Developer/Sandbox/FileSandbox/Source/FileSandboxUI/Public/Persist/Feedback/SummaryPersistFeedback.h b/Engine/Plugins/Developer/Sandbox/FileSandbox/Source/FileSandboxUI/Public/Persist/Feedback/SummaryPersistFeedback.h
new file mode 100644
--- /dev/null
+++ b/Engine/Plugins/Developer/Sandbox/FileSandbox/Source/FileSandboxUI/Public/Persist/Feedback/SummaryPersistFeedback.h
@@ -0,0 +1,56 @@
+﻿// Copyright Epic Games, Inc. All Rights Reserved.
+
+#pragma once
+
+#include "Containers/Map.h"
+#include "PersistSummary.h"
+#include "Interface/IPersistFeedback.h"
+#include "Types/GatheredFileChanges.h"
+
+#define UE_API FILESANDBOXUI_API
+
+namespace UE::FileSandboxCore { enum class ESandboxFileChange : uint8; }
+
+namespace UE::FileSandboxUI
+{
+/**
+ * Fills in a FPersistSummary so you
+ */
+class FSummaryPersistFeedback : public FileSandboxCore::IPersistFeedback
+{
+	TMap<FString, FileSandboxCore::ESandboxFileChange> FileActions;
+	
+public:
+	
+	/** The summary of the operation. */
+	FPersistSummary Summary;
+	
+	UE_API explicit FSummaryPersistFeedback(
+		TConstArrayView<FString> InPersistedFiles,
+		TConstArrayView<FileSandboxCore::ESandboxFileChange> InFileActions
+		);
+	
+	UE_API explicit FSummaryPersistFeedback(const FileSandboxCore::FGatheredFileChanges& InChanges)
+		: FSummaryPersistFeedback(InChanges.NonSandboxPaths, InChanges.FileActions)
+	{}
+
+	//~ Begin IPersistFeedback Interface
+	virtual void StartFile(const TCHAR* InFilename) override { }
+	virtual void HandleSuccess(const TCHAR* InFilename) override { }
+	virtual void HandleError_CheckoutNotAllowed(const TCHAR* InFilename) override { HandleError(InFilename); }
+	virtual void HandleError_Checkout(const TCHAR* InFilename) override { HandleError(InFilename); }
+	virtual void HandleError_Revert(const TCHAR* InFilename) override { HandleError(InFilename); }
+	virtual void HandleError_MarkForAdd(const TCHAR* InFilename) override { HandleError(InFilename); }
+	virtual void HandleError_DeleteSCC(const TCHAR* InFilename) override { HandleError(InFilename); }
+	virtual void HandleError_MakeWritable(const TCHAR* InFilename) override { HandleError(InFilename); }
+	virtual void HandleError_MoveFile(const TCHAR* InToFilename, const TCHAR* InFromFilename) override { HandleError(InToFilename); }
+	virtual void HandleError_DeleteFile(const TCHAR* InFilename) override { HandleError(InFilename); }
+	//~ End IPersistFeedback Interface
+	
+private:
+	
+	UE_API void HandleError(const TCHAR* InToFilename);
+};
+}
+
+#undef UE_API
\ No newline at end of file
diff --git a/Engine/Plugins/Developer/Sandbox/SandboxedEditing/Source/SandboxedEditing/Private/Features/Browser/ViewModels/Persist/PersistSandboxViewModel.cpp b/Engine/Plugins/Developer/Sandbox/SandboxedEditing/Source/SandboxedEditing/Private/Features/Browser/ViewModels/Persist/PersistSandboxViewModel.cpp
--- a/Engine/Plugins/Developer/Sandbox/SandboxedEditing/Source/SandboxedEditing/Private/Features/Browser/ViewModels/Persist/PersistSandboxViewModel.cpp
+++ b/Engine/Plugins/Developer/Sandbox/SandboxedEditing/Source/SandboxedEditing/Private/Features/Browser/ViewModels/Persist/PersistSandboxViewModel.cpp
@@ -3,6 +3,9 @@
 #include "PersistSandboxViewModel.h"
 
 #include "PersistWorkflowUserChoice.h"
+#include "Interface/Feedback/AggregatePersistFeedback.h"
+#include "Persist/Feedback/SlowTaskPersistFeedback.h"
+#include "Persist/Feedback/SummaryPersistFeedback.h"
 
 namespace UE::SandboxedEditing
 {
@@ -33,7 +36,13 @@ void FPersistSandboxViewModel::OnCleanupPersistWorkflow(const FPersistWorkflowUs
 {
 	if (InPersistResult.UserAction == EPersistWorkflowAction::Persist)
 	{
-		Model->PersistFiles(*InPersistResult.FilesToPersist);
+		FileSandboxUI::FSlowTaskPersistFeedback ProgressBarFeedback(InPersistResult.FilesToPersist->Num());
+		ProgressBarFeedback.SlowTask.MakeDialogDelayed(0.2f);
+		FileSandboxUI::FSummaryPersistFeedback SummaryFeedback(ActivePersistWorkflow->GetPersistableFiles());
+		FileSandboxCore::FAggregatePersistFeedback Aggregate({ &ProgressBarFeedback, &SummaryFeedback });
+		Model->PersistFiles(*InPersistResult.FilesToPersist, &Aggregate);
+		
+		OnRequestPersistSummaryNotificationDelegate.Broadcast(SummaryFeedback.Summary);
 	}
 	
 	ActivePersistWorkflow.Reset();
diff --git a/Engine/Plugins/Developer/Sandbox/SandboxedEditing/Source/SandboxedEditing/Private/Features/Browser/ViewModels/Persist/PersistSandboxViewModel.h b/Engine/Plugins/Developer/Sandbox/SandboxedEditing/Source/SandboxedEditing/Private/Features/Browser/ViewModels/Persist/PersistSandboxViewModel.h
--- a/Engine/Plugins/Developer/Sandbox/SandboxedEditing/Source/SandboxedEditing/Private/Features/Browser/ViewModels/Persist/PersistSandboxViewModel.h
+++ b/Engine/Plugins/Developer/Sandbox/SandboxedEditing/Source/SandboxedEditing/Private/Features/Browser/ViewModels/Persist/PersistSandboxViewModel.h
@@ -6,11 +6,14 @@
 #include "Framework/Models/SandboxSystemModel.h"
 #include "Templates/UnrealTemplate.h"
 
+namespace UE::FileSandboxUI { struct FPersistSummary; }
+
 namespace UE::SandboxedEditing
 {
 enum class EPersistWorkflowAction : uint8;
 DECLARE_MULTICAST_DELEGATE_OneParam(FPersistSandboxWorkflowDelegate, FPersistSandboxWorkflow&)
 DECLARE_MULTICAST_DELEGATE_OneParam(FPersistWorkflowActionDelegate, EPersistWorkflowAction);
+DECLARE_MULTICAST_DELEGATE_OneParam(FRequestPersistSummaryNotificationDelegate, const FileSandboxUI::FPersistSummary&)
 
 /** Handles the logic of persisting a sandbox. */
 class FPersistSandboxViewModel : public FNoncopyable
@@ -33,6 +36,9 @@ public:
 	/** Invoked when the persist workflow has ended. */
 	FPersistWorkflowActionDelegate& OnEndPersistWorkflow() { return OnEndPersistWorkflowDelegate; }
 	
+	/** Invoked when a persist operation is performed to show a notification for the result. */
+	FRequestPersistSummaryNotificationDelegate& OnRequestPersistSummaryNotification() { return OnRequestPersistSummaryNotificationDelegate; }
+	
 private:
 	
 	/** Used to persist changes via the sandbox system. */
@@ -46,6 +52,9 @@ private:
 	/** Invoked when the persist workflow has ended. */
 	FPersistWorkflowActionDelegate OnEndPersistWorkflowDelegate;
 	
+	/** Invoked when a persist operation is performed to show a notification for the result. */
+	FRequestPersistSummaryNotificationDelegate OnRequestPersistSummaryNotificationDelegate;
+	
 	/** Handles the workflow ending, possibly persisting InFilesToPersist. */
 	void OnCleanupPersistWorkflow(const FPersistWorkflowUserChoice& InPersistResult);
 };
diff --git a/Engine/Plugins/Developer/Sandbox/SandboxedEditing/Source/SandboxedEditing/Private/Features/Browser/Views/PersistSandboxView.cpp b/Engine/Plugins/Developer/Sandbox/SandboxedEditing/Source/SandboxedEditing/Private/Features/Browser/Views/PersistSandboxView.cpp
--- a/Engine/Plugins/Developer/Sandbox/SandboxedEditing/Source/SandboxedEditing/Private/Features/Browser/Views/PersistSandboxView.cpp
+++ b/Engine/Plugins/Developer/Sandbox/SandboxedEditing/Source/SandboxedEditing/Private/Features/Browser/Views/PersistSandboxView.cpp
@@ -6,7 +6,9 @@
 #include "Features/Browser/ViewModels/Persist/PersistSandboxViewModel.h"
 #include "Features/Browser/Widgets/Sandboxed/SPersistOperationWidget.h"
 #include "Framework/Application/SlateApplication.h"
+#include "Framework/Notifications/NotificationManager.h"
 #include "Widgets/SWindow.h"
+#include "Widgets/Notifications/SNotificationList.h"
 
 #define LOCTEXT_NAMESPACE "FPersistSandboxView"
 
@@ -16,14 +18,16 @@ FPersistSandboxView::FPersistSandboxView(const TSharedRef<FPersistSandboxViewMod
 	: PersistViewModel(InPersistViewModel)
 {
 	PersistViewModel->OnStartPersistWorkflow().AddRaw(this, &FPersistSandboxView::ShowPersistWorkflow);
+	PersistViewModel->OnRequestPersistSummaryNotification().AddRaw(this, &FPersistSandboxView::ShowPersistSummary);
 }
 
 FPersistSandboxView::~FPersistSandboxView()
 {
 	PersistViewModel->OnStartPersistWorkflow().RemoveAll(this);
+	PersistViewModel->OnRequestPersistSummaryNotification().RemoveAll(this);
 }
 
-void FPersistSandboxView::ShowPersistWorkflow(FPersistSandboxWorkflow& InWorkflow)
+void FPersistSandboxView::ShowPersistWorkflow(FPersistSandboxWorkflow& InWorkflow) const
 {
 	const TSharedRef<SWindow> NewWindow = SNew(SWindow)
 		.Title(LOCTEXT("PersistSubmitWindowTitle", "Persist Sandbox"))
@@ -51,6 +55,51 @@ void FPersistSandboxView::ShowPersistWorkflow(FPersistSandboxWorkflow& InWorkflo
 		Workflow->CancelPersist();
 	}
 }
+
+namespace PersistViewDetail
+{
+static FText BuildErrorSubtext(const FileSandboxUI::FPersistSummary& InSummary)
+{
+	return FText::Format(LOCTEXT("Persist.SubText.WithErrorsFmt", "{0} persisted. {1} errors.\nSee log for details."), 
+		InSummary.Edited + InSummary.Added + InSummary.Deleted,
+		InSummary.FailedFiles.Num()
+		);
+}
+static FText BuildSuccessSubtext(const FileSandboxUI::FPersistSummary& InSummary)
+{
+	return FText::Format(
+		LOCTEXT("Persist.SubText.SuccessFmt", "{0}{1}{2}"),
+		InSummary.Edited > 0 ? FText::Format(LOCTEXT("Edited", "Edited {0} files.\n"), InSummary.Edited) : FText::GetEmpty(),
+		InSummary.Added > 0 ? FText::Format(LOCTEXT("Added", "Added {0} files.\n"), InSummary.Added) : FText::GetEmpty(),
+		InSummary.Deleted > 0 ? FText::Format(LOCTEXT("Removed", "Removed {0} files."), InSummary.Deleted) : FText::GetEmpty()
+		);
+}
+}
+
+void FPersistSandboxView::ShowPersistSummary(const FileSandboxUI::FPersistSummary& InSummary) const
+{
+	FSlateNotificationManager& NotificationManager = FSlateNotificationManager::Get();
+
+	const bool bHasErrors = !InSummary.FailedFiles.IsEmpty();
+	const FText Title = bHasErrors 
+		? LOCTEXT("Persist.Title.WithErrors", "Persist completed with errors")
+		: LOCTEXT("Persist.Title.Success", "Persist successful");
+	
+	FNotificationInfo Info(Title);
+	Info.bFireAndForget = false;
+	Info.SubText = bHasErrors ? PersistViewDetail::BuildErrorSubtext(InSummary) : PersistViewDetail::BuildSuccessSubtext(InSummary);
+	
+	const TSharedPtr<SNotificationItem> Item = NotificationManager.AddNotification(Info);
+	Item->SetCompletionState(bHasErrors ? SNotificationItem::CS_Fail : SNotificationItem::CS_Success);
+	Item->RegisterActiveTimer(3.f, FWidgetActiveTimerDelegate::CreateLambda([WeakItem = Item.ToWeakPtr()](auto, auto)
+	{
+		if (const TSharedPtr<SNotificationItem> ItemPin = WeakItem.Pin())
+		{
+			ItemPin->ExpireAndFadeout();
+		}
+		return EActiveTimerReturnType::Stop;
+	}));
+}
 }
 
 #undef LOCTEXT_NAMESPACE
\ No newline at end of file
diff --git a/Engine/Plugins/Developer/Sandbox/SandboxedEditing/Source/SandboxedEditing/Private/Features/Browser/Views/PersistSandboxView.h b/Engine/Plugins/Developer/Sandbox/SandboxedEditing/Source/SandboxedEditing/Private/Features/Browser/Views/PersistSandboxView.h
--- a/Engine/Plugins/Developer/Sandbox/SandboxedEditing/Source/SandboxedEditing/Private/Features/Browser/Views/PersistSandboxView.h
+++ b/Engine/Plugins/Developer/Sandbox/SandboxedEditing/Source/SandboxedEditing/Private/Features/Browser/Views/PersistSandboxView.h
@@ -3,6 +3,7 @@
 #pragma once
 
 #include "Features/Browser/ViewModels/Persist/PersistSandboxWorkflow.h"
+#include "Persist/Feedback/PersistSummary.h"
 #include "Templates/SharedPointer.h"
 #include "Templates/UnrealTemplate.h"
 
@@ -14,7 +15,6 @@ class FPersistSandboxViewModel;
 class FPersistSandboxView : public FNoncopyable
 {
 public:
-	
 	explicit FPersistSandboxView(const TSharedRef<FPersistSandboxViewModel>& InPersistViewModel);
 	~FPersistSandboxView();
 	
@@ -23,7 +23,10 @@ private:
 	/** Notifies when FPersistSandboxWorkflow starts. */
 	const TSharedRef<FPersistSandboxViewModel> PersistViewModel;
 	
-	void ShowPersistWorkflow(FPersistSandboxWorkflow& InWorkflow);
+	/** Shows a modal in which the user can select the files to persist. */
+	void ShowPersistWorkflow(FPersistSandboxWorkflow& InWorkflow) const;
+	/** Shows a notification summarizing the result of the operation. */
+	void ShowPersistSummary(const FileSandboxUI::FPersistSummary& InSummary) const;
 };
 }
 
diff --git a/Engine/Plugins/Developer/Sandbox/SandboxedEditing/Source/SandboxedEditing/Private/Framework/Models/SandboxSystemModel.cpp b/Engine/Plugins/Developer/Sandbox/SandboxedEditing/Source/SandboxedEditing/Private/Framework/Models/SandboxSystemModel.cpp
--- a/Engine/Plugins/Developer/Sandbox/SandboxedEditing/Source/SandboxedEditing/Private/Framework/Models/SandboxSystemModel.cpp
+++ b/Engine/Plugins/Developer/Sandbox/SandboxedEditing/Source/SandboxedEditing/Private/Framework/Models/SandboxSystemModel.cpp
@@ -229,14 +229,19 @@ void FSandboxSystemModel::PersistAllChanges() const
 	}
 }
 
-void FSandboxSystemModel::PersistFiles(TConstArrayView<FString> InFiles) const
+void FSandboxSystemModel::PersistFiles(TConstArrayView<FString> InFiles, FileSandboxCore::IPersistFeedback* InFeedback) const
 {
 	using namespace FileSandboxCore;
 	IFileSandboxCoreModule& Module = IFileSandboxCoreModule::Get();
 	ISandboxManager& SandboxManager = Module.GetSandboxManager();
 	if (ISandboxInstance* Instance = SandboxManager.GetActiveSandboxInstance())
 	{
-		Instance->PersistSandbox(FPersistArgs{ InFiles });
+		const FPersistArgs Args 
+		{
+			.Files = InFiles,
+			.Feedback = InFeedback
+		};
+		Instance->PersistSandbox(Args);
 	}
 }
 
diff --git a/Engine/Plugins/Developer/Sandbox/SandboxedEditing/Source/SandboxedEditing/Private/Framework/Models/SandboxSystemModel.h b/Engine/Plugins/Developer/Sandbox/SandboxedEditing/Source/SandboxedEditing/Private/Framework/Models/SandboxSystemModel.h
--- a/Engine/Plugins/Developer/Sandbox/SandboxedEditing/Source/SandboxedEditing/Private/Framework/Models/SandboxSystemModel.h
+++ b/Engine/Plugins/Developer/Sandbox/SandboxedEditing/Source/SandboxedEditing/Private/Framework/Models/SandboxSystemModel.h
@@ -8,6 +8,7 @@
 
 namespace UE::FileSandboxCore
 {
+class IPersistFeedback;
 struct FGatheredFileChanges;
 class ISandboxInstance;
 }
@@ -68,7 +69,7 @@ public:
 	/** Persists all changes */
 	void PersistAllChanges() const;
 	/** Persists given files. */
-	void PersistFiles(TConstArrayView<FString> InFiles) const;
+	void PersistFiles(TConstArrayView<FString> InFiles, FileSandboxCore::IPersistFeedback* InFeedback = nullptr) const;
 	/** @return Whether all changes can be persisted */
 	bool CanPersistAllChanges() const;
 	
