File Sandbox Core: Fix loca issues

#Synced-CL 49159615

diff --git a/Engine/Plugins/Developer/Sandbox/SandboxedEditing/Source/SandboxedEditing/Private/Features/Browser/Widgets/Sandboxed/FileState/Columns/ActionFileStateColumn.cpp b/Engine/Plugins/Developer/Sandbox/SandboxedEditing/Source/SandboxedEditing/Private/Features/Browser/Widgets/Sandboxed/FileState/Columns/ActionFileStateColumn.cpp
--- a/Engine/Plugins/Developer/Sandbox/SandboxedEditing/Source/SandboxedEditing/Private/Features/Browser/Widgets/Sandboxed/FileState/Columns/ActionFileStateColumn.cpp
+++ b/Engine/Plugins/Developer/Sandbox/SandboxedEditing/Source/SandboxedEditing/Private/Features/Browser/Widgets/Sandboxed/FileState/Columns/ActionFileStateColumn.cpp
@@ -33,8 +33,8 @@ static FText GetActionToolTipText(FileSandboxCore::ESandboxFileChange InChange)
 	switch (InChange)
 	{
 	case FileSandboxCore::ESandboxFileChange::Added: return LOCTEXT("Added", "Marked for add");
-	case FileSandboxCore::ESandboxFileChange::Removed: return LOCTEXT("Added", "Marked for delete");
-	case FileSandboxCore::ESandboxFileChange::Edited: return LOCTEXT("Added", "File is edited");
+	case FileSandboxCore::ESandboxFileChange::Removed: return LOCTEXT("Removed", "Marked for delete");
+	case FileSandboxCore::ESandboxFileChange::Edited: return LOCTEXT("Edited", "File is edited");
 	case FileSandboxCore::ESandboxFileChange::None: [[fallthrough]];
 	default: return FText::GetEmpty();
 	}
