Sandboxed Editing: Add file state to browser

#Synced-CL 49159615

diff --git a/Engine/Plugins/Developer/Sandbox/SandboxedEditing/Source/SandboxedEditing/Private/Features/Browser/BrowserFeature.cpp b/Engine/Plugins/Developer/Sandbox/SandboxedEditing/Source/SandboxedEditing/Private/Features/Browser/BrowserFeature.cpp
--- a/Engine/Plugins/Developer/Sandbox/SandboxedEditing/Source/SandboxedEditing/Private/Features/Browser/BrowserFeature.cpp
+++ b/Engine/Plugins/Developer/Sandbox/SandboxedEditing/Source/SandboxedEditing/Private/Features/Browser/BrowserFeature.cpp
@@ -55,7 +55,8 @@ FBrowserFeature::FBrowserFeature(
 	)
 	: SandboxModel(InSandboxModel)
 	, SandboxColumns(BrowserDetail::MakeSandboxColumnRegistry())
-	, ViewModels(SandboxColumns, SandboxModel, InExternalSandboxViewModel)
+	, BrowserColumns(GetColumnsForBrowser())
+	, ViewModels(SandboxColumns, BrowserColumns, SandboxModel, InExternalSandboxViewModel)
 	, Views(ViewModels.PersistViewModel, ViewModels.LeaveViewModel)
 	, CommandBindings(MakeShared<FBrowserCommandBindings>(ViewModels))
 {
@@ -99,7 +100,8 @@ TSharedRef<SDockTab> FBrowserFeature::SpawnBrowserTab(const FSpawnTabArgs& InArg
 		[
 			SNew(SSandboxedEditingRoot, ViewModels)
 			.CommandList(CommandBindings->GetCommandList())
-			.ColumnFactories(SandboxColumns.ColumnFactories)
+			.SandboxColumnFactories(SandboxColumns.ColumnFactories)
+			.FileActionsColumnFactories(BrowserColumns.ColumnFactories)
 		];
 }
 }
diff --git a/Engine/Plugins/Developer/Sandbox/SandboxedEditing/Source/SandboxedEditing/Private/Features/Browser/BrowserFeature.h b/Engine/Plugins/Developer/Sandbox/SandboxedEditing/Source/SandboxedEditing/Private/Features/Browser/BrowserFeature.h
--- a/Engine/Plugins/Developer/Sandbox/SandboxedEditing/Source/SandboxedEditing/Private/Features/Browser/BrowserFeature.h
+++ b/Engine/Plugins/Developer/Sandbox/SandboxedEditing/Source/SandboxedEditing/Private/Features/Browser/BrowserFeature.h
@@ -46,6 +46,9 @@ private:
 	
 	/** The columns that are displayed in the browser. */
 	const FSandboxColumnRegistry SandboxColumns;
+	/** The columns displayed for the file actions in the browser. */
+	const FFileStateColumnRegistry BrowserColumns;
+	
 	/** Holds all view models relevant for the browser. */
 	const FBrowserViewModels ViewModels;
 	/** Holds all explicit views. */
diff --git a/Engine/Plugins/Developer/Sandbox/SandboxedEditing/Source/SandboxedEditing/Private/Features/Browser/ViewModels/BrowserViewModels.h b/Engine/Plugins/Developer/Sandbox/SandboxedEditing/Source/SandboxedEditing/Private/Features/Browser/ViewModels/BrowserViewModels.h
--- a/Engine/Plugins/Developer/Sandbox/SandboxedEditing/Source/SandboxedEditing/Private/Features/Browser/ViewModels/BrowserViewModels.h
+++ b/Engine/Plugins/Developer/Sandbox/SandboxedEditing/Source/SandboxedEditing/Private/Features/Browser/ViewModels/BrowserViewModels.h
@@ -4,13 +4,15 @@
 
 #include "Active/ActiveSandboxDetailsViewModel.h"
 #include "Active/ActiveSandboxTrackerViewModel.h"
+#include "FileState/FileStateColumnRegistry.h"
+#include "FileState/UnloadedSandboxFileStateViewModel.h"
 #include "Leaving/LeaveSandboxViewModel.h"
 #include "List/FilterSandboxViewModel.h"
+#include "List/SandboxColumnRegistry.h"
 #include "List/SandboxListViewModel.h"
 #include "Persist/PersistSandboxViewModel.h"
 #include "SandboxControlsViewModel.h"
 #include "SandboxMetaDataViewModel.h"
-#include "List/SandboxColumnRegistry.h"
 #include "Templates/SharedPointer.h"
 
 namespace UE::FileSandboxUI { class IExternalSandboxActiveViewModel; }
@@ -26,12 +28,15 @@ struct FBrowserViewModels
 	const TSharedRef<FFilterSandboxViewModel> FilterSandboxViewModel;
 	/** Knows about the items that are displayed in the browser. */
 	const TSharedRef<FSandboxListViewModel> ListViewModel;
+	/** Knows the file action of the currently selected sandbox. */
+	const TSharedRef<FUnloadedSandboxFileStateViewModel> UnloadedSandboxFileStateViewModel;
 	
 	/** Handles actions that the user performs on the sandboxes via the UI. */
 	const TSharedRef<FSandboxControlsViewModel> ControlsViewModel;
 	/** Handles changing metadata of a sandbox. */
 	const TSharedRef<FSandboxMetaDataViewModel> MetaDataViewModel;
 	
+	
 	/** Handles persisting the active sandbox. */
 	const TSharedRef<FPersistSandboxViewModel> PersistViewModel;
 	/** Handles leaving an active sandbox. */
@@ -46,11 +51,13 @@ struct FBrowserViewModels
 	
 	explicit FBrowserViewModels(
 		const FSandboxColumnRegistry& InSandboxColumns, 
+		const FFileStateColumnRegistry& InBrowserColumns,
 		const TSharedRef<FSandboxSystemModel>& InModel,
 		const TSharedRef<FileSandboxUI::IExternalSandboxActiveViewModel>& InExternalSandboxViewModel
 		)
 		: FilterSandboxViewModel(MakeShared<FFilterSandboxViewModel>(InSandboxColumns.GetBehaviorArray()))
 		, ListViewModel(MakeShared<FSandboxListViewModel>(InModel, FilterSandboxViewModel, InSandboxColumns.ColumnBehaviors))
+		, UnloadedSandboxFileStateViewModel(MakeShared<FUnloadedSandboxFileStateViewModel>(InModel, InBrowserColumns.ColumnBehaviors))
 		, ControlsViewModel(MakeShared<FSandboxControlsViewModel>(InModel))
 		, MetaDataViewModel(MakeShared<FSandboxMetaDataViewModel>(InModel))
 		, PersistViewModel(MakeShared<FPersistSandboxViewModel>(InModel))
diff --git a/Engine/Plugins/Developer/Sandbox/SandboxedEditing/Source/SandboxedEditing/Private/Features/Browser/ViewModels/FileState/FileStateColumnRegistry.cpp b/Engine/Plugins/Developer/Sandbox/SandboxedEditing/Source/SandboxedEditing/Private/Features/Browser/ViewModels/FileState/FileStateColumnRegistry.cpp
new file mode 100644
--- /dev/null
+++ b/Engine/Plugins/Developer/Sandbox/SandboxedEditing/Source/SandboxedEditing/Private/Features/Browser/ViewModels/FileState/FileStateColumnRegistry.cpp
@@ -0,0 +1,54 @@
+﻿// Copyright Epic Games, Inc. All Rights Reserved.
+
+#include "FileStateColumnRegistry.h"
+
+#include "FileStateColumns.h"
+#include "Features/Browser/Widgets/Sandboxed/FileState/Columns/ActionFileStateColumn.h"
+#include "Features/Browser/Widgets/Sandboxed/FileState/Columns/PathFileStateColumn.h"
+#include "Features/Browser/Widgets/Sandboxed/FileState/Columns/PersistFileStateColumn.h"
+
+namespace UE::SandboxedEditing
+{
+class FPathFileStateColumn;
+
+FFileStateColumnRegistry GetColumnsForBrowser()
+{
+	const TSharedRef<FActionFileStateColumn> ActionColum = MakeShared<FActionFileStateColumn>();
+	const TSharedRef<FPathFileStateColumn> PathColumn = MakeShared<FPathFileStateColumn>();
+	
+	FFileStateColumnRegistry Registry;
+	Registry.ColumnBehaviors = TMap<FName, TSharedRef<IFileStateColumnBehavior>>
+	{
+		{ FileActionColumn, ActionColum },
+		{ FilePathColumn, PathColumn }
+	};
+	Registry.ColumnFactories = TMap<FName, TSharedRef<IFileStateColumnWidgetFactory>>
+	{
+		{ FileActionColumn, ActionColum },
+		{ FilePathColumn, PathColumn }
+	};
+	return Registry;
+}
+
+FFileStateColumnRegistry GetColumnsForPersist(const TSharedRef<FPersistOperationViewModel>& InViewModel)
+{
+	const TSharedRef<FPersistFileStateColumn> CheckboxColumn = MakeShared<FPersistFileStateColumn>(InViewModel);
+	const TSharedRef<FActionFileStateColumn> ActionColum = MakeShared<FActionFileStateColumn>();
+	const TSharedRef<FPathFileStateColumn> PathColumn = MakeShared<FPathFileStateColumn>();
+	
+	FFileStateColumnRegistry Registry;
+	Registry.ColumnBehaviors = TMap<FName, TSharedRef<IFileStateColumnBehavior>>
+	{
+		{ PersistCheckboxColumn, CheckboxColumn },
+		{ FileActionColumn, ActionColum },
+		{ FilePathColumn, PathColumn }
+	};
+	Registry.ColumnFactories = TMap<FName, TSharedRef<IFileStateColumnWidgetFactory>>
+	{
+		{ PersistCheckboxColumn, CheckboxColumn },
+		{ FileActionColumn, ActionColum },
+		{ FilePathColumn, PathColumn }
+	};
+	return Registry;
+}
+}
diff --git a/Engine/Plugins/Developer/Sandbox/SandboxedEditing/Source/SandboxedEditing/Private/Features/Browser/ViewModels/FileState/FileStateColumnRegistry.h b/Engine/Plugins/Developer/Sandbox/SandboxedEditing/Source/SandboxedEditing/Private/Features/Browser/ViewModels/FileState/FileStateColumnRegistry.h
new file mode 100644
--- /dev/null
+++ b/Engine/Plugins/Developer/Sandbox/SandboxedEditing/Source/SandboxedEditing/Private/Features/Browser/ViewModels/FileState/FileStateColumnRegistry.h
@@ -0,0 +1,27 @@
+﻿// Copyright Epic Games, Inc. All Rights Reserved.
+
+#pragma once
+
+#include "Containers/Map.h"
+#include "Templates/SharedPointer.h"
+
+namespace UE::SandboxedEditing
+{
+class FPersistOperationViewModel;
+class IFileStateColumnBehavior;
+class IFileStateColumnWidgetFactory;
+
+struct FFileStateColumnRegistry
+{
+	/** The behavior of the columns for view models. */
+	TMap<FName, TSharedRef<IFileStateColumnBehavior>> ColumnBehaviors;
+	
+	/** Creates the widgets for the columns. */
+	TMap<FName, TSharedRef<IFileStateColumnWidgetFactory>> ColumnFactories;
+};
+	
+/** @return Columns to show in the browser view. */
+FFileStateColumnRegistry GetColumnsForBrowser();
+/** @return Columns to show in the persist view. */
+FFileStateColumnRegistry GetColumnsForPersist(const TSharedRef<FPersistOperationViewModel>& InViewModel);
+}
\ No newline at end of file
diff --git a/Engine/Plugins/Developer/Sandbox/SandboxedEditing/Source/SandboxedEditing/Private/Features/Browser/ViewModels/FileState/FileStateItemUtils.cpp b/Engine/Plugins/Developer/Sandbox/SandboxedEditing/Source/SandboxedEditing/Private/Features/Browser/ViewModels/FileState/FileStateItemUtils.cpp
new file mode 100644
--- /dev/null
+++ b/Engine/Plugins/Developer/Sandbox/SandboxedEditing/Source/SandboxedEditing/Private/Features/Browser/ViewModels/FileState/FileStateItemUtils.cpp
@@ -0,0 +1,25 @@
+﻿// Copyright Epic Games, Inc. All Rights Reserved.
+
+#include "FileStateItemUtils.h"
+
+#include "FileStateItem.h"
+#include "Templates/SharedPointer.h"
+#include "Types/GatheredFileChanges.h"
+#include "Types/SandboxFileChange.h"
+
+namespace UE::SandboxedEditing
+{
+TArray<TSharedPtr<FFileStateItem>> MakeItemsFromFileChanges(const FileSandboxCore::FGatheredFileChanges& InChanges)
+{
+	const TArray<FString>& NonSandboxPaths = InChanges.NonSandboxPaths;
+	const TArray<FileSandboxCore::ESandboxFileChange>& FileActions = InChanges.FileActions;
+	check(NonSandboxPaths.Num() == FileActions.Num());
+	
+	TArray<TSharedPtr<FFileStateItem>> Items;
+	for (int32 Index = 0; Index < NonSandboxPaths.Num(); ++Index)
+	{
+		Items.Emplace(MakeShared<FFileStateItem>(NonSandboxPaths[Index], FileActions[Index]));
+	}
+	return Items;
+}
+}
diff --git a/Engine/Plugins/Developer/Sandbox/SandboxedEditing/Source/SandboxedEditing/Private/Features/Browser/ViewModels/FileState/FileStateItemUtils.h b/Engine/Plugins/Developer/Sandbox/SandboxedEditing/Source/SandboxedEditing/Private/Features/Browser/ViewModels/FileState/FileStateItemUtils.h
new file mode 100644
--- /dev/null
+++ b/Engine/Plugins/Developer/Sandbox/SandboxedEditing/Source/SandboxedEditing/Private/Features/Browser/ViewModels/FileState/FileStateItemUtils.h
@@ -0,0 +1,17 @@
+﻿// Copyright Epic Games, Inc. All Rights Reserved.
+
+#pragma once
+
+#include "Containers/ContainersFwd.h"
+#include "Templates/SharedPointerFwd.h"
+
+namespace UE::FileSandboxCore { struct FGatheredFileChanges; }
+
+namespace UE::SandboxedEditing
+{
+struct FFileStateItem;
+
+/** @return An item array transformed from FGatheredFileChanges  */
+TArray<TSharedPtr<FFileStateItem>> MakeItemsFromFileChanges(const FileSandboxCore::FGatheredFileChanges& InChanges);
+}
+
diff --git a/Engine/Plugins/Developer/Sandbox/SandboxedEditing/Source/SandboxedEditing/Private/Features/Browser/ViewModels/FileState/UnloadedSandboxFileStateViewModel.cpp b/Engine/Plugins/Developer/Sandbox/SandboxedEditing/Source/SandboxedEditing/Private/Features/Browser/ViewModels/FileState/UnloadedSandboxFileStateViewModel.cpp
new file mode 100644
--- /dev/null
+++ b/Engine/Plugins/Developer/Sandbox/SandboxedEditing/Source/SandboxedEditing/Private/Features/Browser/ViewModels/FileState/UnloadedSandboxFileStateViewModel.cpp
@@ -0,0 +1,23 @@
+﻿// Copyright Epic Games, Inc. All Rights Reserved.
+
+#include "UnloadedSandboxFileStateViewModel.h"
+
+#include "FileStateItemUtils.h"
+#include "Types/GatheredFileChanges.h"
+
+namespace UE::SandboxedEditing
+{
+void FUnloadedSandboxFileStateViewModel::SetContent(const FString& InSandboxRoot)
+{
+	Items = MakeItemsFromFileChanges(
+		Model->GatherFileChanges(InSandboxRoot)
+		);
+	OnItemsChangedDelegate.Broadcast();
+}
+
+void FUnloadedSandboxFileStateViewModel::ClearContent()
+{
+	Items.Empty();
+	OnItemsChangedDelegate.Broadcast();
+}
+}
diff --git a/Engine/Plugins/Developer/Sandbox/SandboxedEditing/Source/SandboxedEditing/Private/Features/Browser/ViewModels/FileState/UnloadedSandboxFileStateViewModel.h b/Engine/Plugins/Developer/Sandbox/SandboxedEditing/Source/SandboxedEditing/Private/Features/Browser/ViewModels/FileState/UnloadedSandboxFileStateViewModel.h
new file mode 100644
--- /dev/null
+++ b/Engine/Plugins/Developer/Sandbox/SandboxedEditing/Source/SandboxedEditing/Private/Features/Browser/ViewModels/FileState/UnloadedSandboxFileStateViewModel.h
@@ -0,0 +1,38 @@
+﻿// Copyright Epic Games, Inc. All Rights Reserved.
+
+#pragma once
+
+#include "StaticFileStateViewModel.h"
+#include "Framework/Models/SandboxSystemModel.h"
+
+namespace UE::SandboxedEditing
+{
+/**
+ * View model for displaying the file changes of an unloaded sandbox, i.e. when the engine is not sandboxed.
+ * Used to display info about a sandbox in the browser.
+ */
+class FUnloadedSandboxFileStateViewModel : public FStaticFileStateViewModel
+{
+public:
+	
+	explicit FUnloadedSandboxFileStateViewModel(
+		const TSharedRef<FSandboxSystemModel>& InModel,
+		const TMap<FName, TSharedRef<IFileStateColumnBehavior>>& InColumns
+		)
+		: FStaticFileStateViewModel(InColumns, {})
+		, Model(InModel)
+	{}
+	
+	/** Sets the content to be that of the specified sandbox. */
+	void SetContent(const FString& InSandboxRoot);
+	
+	/** Clears the content. */
+	void ClearContent();
+	
+private:
+	
+	/** The model through which we interact with the sandbox system. */
+	const TSharedRef<FSandboxSystemModel> Model;
+};
+}
+
diff --git a/Engine/Plugins/Developer/Sandbox/SandboxedEditing/Source/SandboxedEditing/Private/Features/Browser/ViewModels/Leaving/LeaveSandboxViewModel.cpp b/Engine/Plugins/Developer/Sandbox/SandboxedEditing/Source/SandboxedEditing/Private/Features/Browser/ViewModels/Leaving/LeaveSandboxViewModel.cpp
--- a/Engine/Plugins/Developer/Sandbox/SandboxedEditing/Source/SandboxedEditing/Private/Features/Browser/ViewModels/Leaving/LeaveSandboxViewModel.cpp
+++ b/Engine/Plugins/Developer/Sandbox/SandboxedEditing/Source/SandboxedEditing/Private/Features/Browser/ViewModels/Leaving/LeaveSandboxViewModel.cpp
@@ -46,7 +46,7 @@ void FLeaveSandboxViewModel::LeaveSandbox()
 		return;
 	}
 	
-	const FileSandboxCore::FGatheredFileChanges FileChanges = Model->GatherFileChanges();
+	const FileSandboxCore::FGatheredFileChanges FileChanges = Model->GatherActiveFileChanges();
 	if (FileChanges.HasChanges())
 	{
 		AskUserToPersistWorkflow = MakeUnique<FAskUserToPersistWorkflow>(
@@ -82,7 +82,7 @@ void FLeaveSandboxViewModel::OnUserDecidedWhetherToPersist(EAskUserToPersistResu
 		// Subscribe before starting the workflow because it can complete "immediately" (due to modal window).
 		PersistViewModel->OnEndPersistWorkflow().AddRaw(this, &FLeaveSandboxViewModel::OnPersistOperationEnded);
 		
-		PersistViewModel->StartPersistWorkflow(Model->GatherFileChanges());
+		PersistViewModel->StartPersistWorkflow(Model->GatherActiveFileChanges());
 		
 		// If the workflow immediately ends (due to modal window) or fails, unsubscribe again
 		if (!PersistViewModel->GetPersistWorkflow())
diff --git a/Engine/Plugins/Developer/Sandbox/SandboxedEditing/Source/SandboxedEditing/Private/Features/Browser/Widgets/SSandboxedEditingRoot.cpp b/Engine/Plugins/Developer/Sandbox/SandboxedEditing/Source/SandboxedEditing/Private/Features/Browser/Widgets/SSandboxedEditingRoot.cpp
--- a/Engine/Plugins/Developer/Sandbox/SandboxedEditing/Source/SandboxedEditing/Private/Features/Browser/Widgets/SSandboxedEditingRoot.cpp
+++ b/Engine/Plugins/Developer/Sandbox/SandboxedEditing/Source/SandboxedEditing/Private/Features/Browser/Widgets/SSandboxedEditingRoot.cpp
@@ -27,11 +27,13 @@ void SSandboxedEditingRoot::Construct(const FArguments& InArgs, const FBrowserVi
 	
 	const SSandboxedEditingBrowser::FViewModels ViewModels 
 	{ 
-		InViewModels.FilterSandboxViewModel, InViewModels.ListViewModel, InViewModels.ControlsViewModel, InViewModels.MetaDataViewModel 
+		InViewModels.FilterSandboxViewModel, InViewModels.ListViewModel, InViewModels.ControlsViewModel, InViewModels.MetaDataViewModel,
+		InViewModels.UnloadedSandboxFileStateViewModel
 	};
 	Browser = SNew(SSandboxedEditingBrowser, ViewModels)
 		.CommandList(CommandBindings)
-		.ColumnFactories(InArgs._ColumnFactories);
+		.SandboxColumnFactories(InArgs._SandboxColumnFactories)
+		.FileActionsColumnFactories(InArgs._FileActionsColumnFactories);
 	
 	UpdateWidgetContent();
 }
diff --git a/Engine/Plugins/Developer/Sandbox/SandboxedEditing/Source/SandboxedEditing/Private/Features/Browser/Widgets/SSandboxedEditingRoot.h b/Engine/Plugins/Developer/Sandbox/SandboxedEditing/Source/SandboxedEditing/Private/Features/Browser/Widgets/SSandboxedEditingRoot.h
--- a/Engine/Plugins/Developer/Sandbox/SandboxedEditing/Source/SandboxedEditing/Private/Features/Browser/Widgets/SSandboxedEditingRoot.h
+++ b/Engine/Plugins/Developer/Sandbox/SandboxedEditing/Source/SandboxedEditing/Private/Features/Browser/Widgets/SSandboxedEditingRoot.h
@@ -2,6 +2,7 @@
 
 #pragma once
 
+#include "Features/Browser/ViewModels/FileState/FileStateColumns.h"
 #include "Features/Browser/ViewModels/List/SandboxColumns.h"
 #include "Sandboxed/SActiveSandbox.h"
 #include "Widgets/DeclarativeSyntaxSupport.h"
@@ -29,7 +30,9 @@ public:
 		/** The command list to bind selection-based commands to. */
 		SLATE_ARGUMENT(TSharedPtr<FUICommandList>, CommandList)
 		/** Used to build the columns. */
-		SLATE_ARGUMENT(FSandboxColumnFactoryMap, ColumnFactories)
+		SLATE_ARGUMENT(FSandboxColumnFactoryMap, SandboxColumnFactories)
+		/** Used to build the columns for the file actions list. */
+		SLATE_ARGUMENT(FFileStateColumnFactoryMap, FileActionsColumnFactories)
 	SLATE_END_ARGS()
 
 	void Construct(const FArguments& InArgs, const FBrowserViewModels& InViewModels);
diff --git a/Engine/Plugins/Developer/Sandbox/SandboxedEditing/Source/SandboxedEditing/Private/Features/Browser/Widgets/Sandboxed/FileState/SFileStateListView.cpp b/Engine/Plugins/Developer/Sandbox/SandboxedEditing/Source/SandboxedEditing/Private/Features/Browser/Widgets/Sandboxed/FileState/SFileStateListView.cpp
--- a/Engine/Plugins/Developer/Sandbox/SandboxedEditing/Source/SandboxedEditing/Private/Features/Browser/Widgets/Sandboxed/FileState/SFileStateListView.cpp
+++ b/Engine/Plugins/Developer/Sandbox/SandboxedEditing/Source/SandboxedEditing/Private/Features/Browser/Widgets/Sandboxed/FileState/SFileStateListView.cpp
@@ -23,6 +23,8 @@ void SFileStateListView::Construct(const FArguments& InArgs, const TSharedRef<IF
 		.SelectionMode(ESelectionMode::None)
 		.HeaderRow(MakeHeaderRow())
 	];
+	
+	ViewModel->OnItemsChanged().AddSP(ListView.ToSharedRef(), &SListView<TSharedPtr<FFileStateItem>>::RequestListRefresh);
 }
 
 TSharedRef<SHeaderRow> SFileStateListView::MakeHeaderRow()
diff --git a/Engine/Plugins/Developer/Sandbox/SandboxedEditing/Source/SandboxedEditing/Private/Features/Browser/Widgets/Sandboxed/SPersistOperationWidget.cpp b/Engine/Plugins/Developer/Sandbox/SandboxedEditing/Source/SandboxedEditing/Private/Features/Browser/Widgets/Sandboxed/SPersistOperationWidget.cpp
--- a/Engine/Plugins/Developer/Sandbox/SandboxedEditing/Source/SandboxedEditing/Private/Features/Browser/Widgets/Sandboxed/SPersistOperationWidget.cpp
+++ b/Engine/Plugins/Developer/Sandbox/SandboxedEditing/Source/SandboxedEditing/Private/Features/Browser/Widgets/Sandboxed/SPersistOperationWidget.cpp
@@ -2,14 +2,15 @@
 
 #include "SPersistOperationWidget.h"
 
-#include "SPositiveActionButton.h"
-#include "Features/Browser/ViewModels/FileState/FileStateItem.h"
+#include "Features/Browser/ViewModels/FileState/FileStateItemUtils.h"
 #include "Features/Browser/ViewModels/FileState/StaticFileStateViewModel.h"
 #include "Features/Browser/ViewModels/Persist/PersistOperationViewModel.h"
-#include "FileState/SFileStateListView.h"
 #include "FileState/Columns/ActionFileStateColumn.h"
 #include "FileState/Columns/PathFileStateColumn.h"
 #include "FileState/Columns/PersistFileStateColumn.h"
+#include "FileState/SFileStateListView.h"
+#include "SPositiveActionButton.h"
+#include "Features/Browser/ViewModels/FileState/FileStateColumnRegistry.h"
 #include "Styling/AppStyle.h"
 #include "Types/GatheredFileChanges.h"
 #include "Widgets/Input/SButton.h"
@@ -21,57 +22,14 @@
 
 namespace UE::SandboxedEditing
 {
-namespace PersistOperationDetail
-{
-static TArray<TSharedPtr<FFileStateItem>> MakeItemsFromFileChanges(const FileSandboxCore::FGatheredFileChanges& InChanges)
-{
-	const TArray<FString>& NonSandboxPaths = InChanges.NonSandboxPaths;
-	const TArray<FileSandboxCore::ESandboxFileChange>& FileActions = InChanges.FileActions;
-	check(NonSandboxPaths.Num() == FileActions.Num());
-	
-	TArray<TSharedPtr<FFileStateItem>> Items;
-	for (int32 Index = 0; Index < NonSandboxPaths.Num(); ++Index)
-	{
-		Items.Emplace(MakeShared<FFileStateItem>(NonSandboxPaths[Index], FileActions[Index]));
-	}
-	return Items;
-}
-
-static void MakePersistColumns(
-	const TSharedRef<FPersistOperationViewModel>& InViewModel,
-	TMap<FName, TSharedRef<IFileStateColumnBehavior>>& OutBehaviors, 
-	TMap<FName, TSharedRef<IFileStateColumnWidgetFactory>>& OutFactories
-	)
-{
-	const TSharedRef<FPersistFileStateColumn> CheckboxColumn = MakeShared<FPersistFileStateColumn>(InViewModel);
-	const TSharedRef<FActionFileStateColumn> ActionColum = MakeShared<FActionFileStateColumn>();
-	const TSharedRef<FPathFileStateColumn> PathColumn = MakeShared<FPathFileStateColumn>();
-	
-	OutBehaviors = TMap<FName, TSharedRef<IFileStateColumnBehavior>>
-	{
-		{ PersistCheckboxColumn, CheckboxColumn },
-		{ FileActionColumn, ActionColum },
-		{ FilePathColumn, PathColumn }
-	};
-	OutFactories = TMap<FName, TSharedRef<IFileStateColumnWidgetFactory>>
-	{
-		{ PersistCheckboxColumn, CheckboxColumn },
-		{ FileActionColumn, ActionColum },
-		{ FilePathColumn, PathColumn }
-	};
-}
-}
-
 void SPersistOperationWidget::Construct(const FArguments& InArgs, const TSharedRef<FPersistOperationViewModel>& InViewModel)
 {
 	ViewModel = InViewModel;
 	
-	TMap<FName, TSharedRef<IFileStateColumnBehavior>> ColumnBehaviors;
-	TMap<FName, TSharedRef<IFileStateColumnWidgetFactory>> ColumnFactories;
-	PersistOperationDetail::MakePersistColumns(ViewModel.ToSharedRef(), ColumnBehaviors, ColumnFactories);
+	const FFileStateColumnRegistry Columns = GetColumnsForPersist(ViewModel.ToSharedRef());
 	const TSharedRef<FStaticFileStateViewModel> FileStateViewModel = MakeShared<FStaticFileStateViewModel>(
-		ColumnBehaviors, 
-		PersistOperationDetail::MakeItemsFromFileChanges(ViewModel->GetPersistableFiles())
+		Columns.ColumnBehaviors, 
+		MakeItemsFromFileChanges(ViewModel->GetPersistableFiles())
 		);
 	
 	ChildSlot
@@ -87,7 +45,7 @@ void SPersistOperationWidget::Construct(const FArguments& InArgs, const TSharedR
 				SNew(SBorder)
 				[
 					SNew(SFileStateListView, FileStateViewModel)
-					.ColumnFactories(ColumnFactories)
+					.ColumnFactories(Columns.ColumnFactories)
 				]
 			]
 
diff --git a/Engine/Plugins/Developer/Sandbox/SandboxedEditing/Source/SandboxedEditing/Private/Features/Browser/Widgets/Unsandboxed/GenericBrowser/SSessionBrowserFrame.cpp b/Engine/Plugins/Developer/Sandbox/SandboxedEditing/Source/SandboxedEditing/Private/Features/Browser/Widgets/Unsandboxed/GenericBrowser/SSessionBrowserFrame.cpp
--- a/Engine/Plugins/Developer/Sandbox/SandboxedEditing/Source/SandboxedEditing/Private/Features/Browser/Widgets/Unsandboxed/GenericBrowser/SSessionBrowserFrame.cpp
+++ b/Engine/Plugins/Developer/Sandbox/SandboxedEditing/Source/SandboxedEditing/Private/Features/Browser/Widgets/Unsandboxed/GenericBrowser/SSessionBrowserFrame.cpp
@@ -6,8 +6,10 @@
 #include "Styling/ToolBarStyle.h"
 #include "Widgets/Layout/SBorder.h"
 #include "Widgets/SBoxPanel.h"
+#include "Widgets/Layout/SBox.h"
 #include "Widgets/Layout/SExpandableArea.h"
 #include "Widgets/Layout/SSeparator.h"
+#include "Widgets/Layout/SSplitter.h"
 #include "Widgets/Text/STextBlock.h"
 
 #define LOCTEXT_NAMESPACE "SSessionBrowserFrame"
@@ -34,15 +36,26 @@ TSharedRef<SVerticalBox> SSessionBrowserFrame::CreateContent(const FArguments& I
 	AppendControlWidget(InArgs, Content);
 	AppendSearchWidget(InArgs, Content);
 	
+	const TSharedRef<SSplitter> Splitter = SNew(SSplitter)
+		.Orientation(Orient_Vertical)
+		.MinimumSlotHeight(80.0f) // Prevent widgets from overlapping.
+		
+		+SSplitter::Slot()
+		.Value(0.6)
+		[
+			InArgs._SessionContent.Widget
+		];
+			
+	
 	checkf(InArgs._SessionContent.Widget != SNullWidget::NullWidget, TEXT("This is a required argument."));
 	Content->AddSlot()
 	.FillHeight(1.0f)
 	.Padding(1.0f, 2.0f)
 	[
-		InArgs._SessionContent.Widget
+		Splitter
 	];
 	
-	AppendDetailsWidget(InArgs, Content);
+	AppendDetailsWidget(InArgs, Splitter);
 	
 	return Content;
 }
@@ -82,23 +95,19 @@ void SSessionBrowserFrame::AppendSearchWidget(const FArguments& InArgs, const TS
 	}
 }
 	
-void SSessionBrowserFrame::AppendDetailsWidget(const FArguments& InArgs, const TSharedRef<SVerticalBox>& Content)
+void SSessionBrowserFrame::AppendDetailsWidget(const FArguments& InArgs, const TSharedRef<SSplitter>& Content)
 {
 	const bool bHasDetailsWidget = InArgs._DetailsContent.Widget != SNullWidget::NullWidget;
 	if (bHasDetailsWidget)
 	{
 		Content->AddSlot()
-	   .AutoHeight()
-	   .Padding(FMargin(2.0f, 0.0f))
-		[
-			SNew(SSeparator)
-		];
-	
-		Content->AddSlot()
-		.AutoHeight()
-		.Padding(FMargin(2.0f, 0.0f))
+		.Value(0.4)
 		[
-			InArgs._DetailsContent.Widget
+			SNew(SBox)
+			.Padding(FMargin(2.0f, 0.0f))
+			[
+				InArgs._DetailsContent.Widget
+			]
 		];
 	}
 }
diff --git a/Engine/Plugins/Developer/Sandbox/SandboxedEditing/Source/SandboxedEditing/Private/Features/Browser/Widgets/Unsandboxed/GenericBrowser/SSessionBrowserFrame.h b/Engine/Plugins/Developer/Sandbox/SandboxedEditing/Source/SandboxedEditing/Private/Features/Browser/Widgets/Unsandboxed/GenericBrowser/SSessionBrowserFrame.h
--- a/Engine/Plugins/Developer/Sandbox/SandboxedEditing/Source/SandboxedEditing/Private/Features/Browser/Widgets/Unsandboxed/GenericBrowser/SSessionBrowserFrame.h
+++ b/Engine/Plugins/Developer/Sandbox/SandboxedEditing/Source/SandboxedEditing/Private/Features/Browser/Widgets/Unsandboxed/GenericBrowser/SSessionBrowserFrame.h
@@ -5,6 +5,7 @@
 #include "Widgets/DeclarativeSyntaxSupport.h"
 #include "Widgets/SCompoundWidget.h"
 
+class SSplitter;
 class SVerticalBox;
 
 namespace UE::SandboxedEditing
@@ -45,7 +46,7 @@ private:
 	/** Optionally adds the search widget (above session content), if specified by args. */
 	void AppendSearchWidget(const FArguments& InArgs, const TSharedRef<SVerticalBox>& Content);
 	/** Optionally adds the controls (below session content), if specified by args. */
-	void AppendDetailsWidget(const FArguments& InArgs, const TSharedRef<SVerticalBox>& Content);
+	void AppendDetailsWidget(const FArguments& InArgs, const TSharedRef<SSplitter>& Content);
 };
 }
 
diff --git a/Engine/Plugins/Developer/Sandbox/SandboxedEditing/Source/SandboxedEditing/Private/Features/Browser/Widgets/Unsandboxed/List/SSandboxListView.cpp b/Engine/Plugins/Developer/Sandbox/SandboxedEditing/Source/SandboxedEditing/Private/Features/Browser/Widgets/Unsandboxed/List/SSandboxListView.cpp
--- a/Engine/Plugins/Developer/Sandbox/SandboxedEditing/Source/SandboxedEditing/Private/Features/Browser/Widgets/Unsandboxed/List/SSandboxListView.cpp
+++ b/Engine/Plugins/Developer/Sandbox/SandboxedEditing/Source/SandboxedEditing/Private/Features/Browser/Widgets/Unsandboxed/List/SSandboxListView.cpp
@@ -22,6 +22,7 @@ void SSandboxListView::Construct(const FArguments& InArgs, const FViewModels& In
 	ControlsViewModel = InViewModels.ControlsViewModel;
 	HighlightTextAttr = InArgs._HighlightText;
 	ColumnFactoryMap = InArgs._ColumnFactories;
+	OnSelectedSandboxChangedDelegate = InArgs._OnSelectedSandboxChanged;
 	
 	ChildSlot
 	[
@@ -30,6 +31,7 @@ void SSandboxListView::Construct(const FArguments& InArgs, const FViewModels& In
 		.OnGenerateRow(this, &SSandboxListView::OnGenerateRow)
 		.OnMouseButtonDoubleClick(this, &SSandboxListView::OnItemDoubleClicked)
 		.HeaderRow(MakeHeaderRow())
+		.OnSelectionChanged(this, &SSandboxListView::OnSelectionChanged)
 	];
 	
 	ControlsViewModel->OnCreationWorkflowStarted().AddSP(this, &SSandboxListView::OnCreationWorkflowStarted);
@@ -154,6 +156,13 @@ void SSandboxListView::OnCreationWorkflowEnded()
 		ListView->RequestListRefresh();
 	}
 }
+
+void SSandboxListView::OnSelectionChanged(TSharedPtr<FSandboxListItem> InItem, ESelectInfo::Type) const
+{
+	OnSelectedSandboxChangedDelegate.ExecuteIfBound(
+		InItem ? TOptional<FString>(InItem->SandboxInfo.SandboxRoot) : TOptional<FString>()
+		);
+}
 }
 
 #undef LOCTEXT_NAMESPACE
\ No newline at end of file
diff --git a/Engine/Plugins/Developer/Sandbox/SandboxedEditing/Source/SandboxedEditing/Private/Features/Browser/Widgets/Unsandboxed/List/SSandboxListView.h b/Engine/Plugins/Developer/Sandbox/SandboxedEditing/Source/SandboxedEditing/Private/Features/Browser/Widgets/Unsandboxed/List/SSandboxListView.h
--- a/Engine/Plugins/Developer/Sandbox/SandboxedEditing/Source/SandboxedEditing/Private/Features/Browser/Widgets/Unsandboxed/List/SSandboxListView.h
+++ b/Engine/Plugins/Developer/Sandbox/SandboxedEditing/Source/SandboxedEditing/Private/Features/Browser/Widgets/Unsandboxed/List/SSandboxListView.h
@@ -20,6 +20,8 @@ class FSandboxListItem;
 class FSandboxListViewModel;
 enum class EBreakBehavior : uint8;
 
+DECLARE_DELEGATE_OneParam(FSandboxDelegate, TOptional<FString> /*InSandboxRoot*/);
+
 /** A list view containing all the sandboxes. */
 class SSandboxListView : public SCompoundWidget
 {
@@ -36,6 +38,8 @@ public:
 		SLATE_ATTRIBUTE(FText, HighlightText)
 		/** Used to build the columns. */
 		SLATE_ARGUMENT(FSandboxColumnFactoryMap, ColumnFactories)
+		/** Invoked when the selected sandbox changes. */
+		SLATE_EVENT(FSandboxDelegate, OnSelectedSandboxChanged)
 	SLATE_END_ARGS()
 
 	void Construct(const FArguments& InArgs, const FViewModels& InViewModels);
@@ -72,6 +76,9 @@ private:
 	 */
 	TOptional<FDummySandboxItemInfo> DummyItem;
 	
+	/** Invoked when the selected sandbox changes. */
+	FSandboxDelegate OnSelectedSandboxChangedDelegate;
+	
 	/** @return The header row to display at the top of the table */
 	TSharedRef<SHeaderRow> MakeHeaderRow();
 	
@@ -85,6 +92,8 @@ private:
 	
 	void OnCreationWorkflowStarted(const FSandboxCreationWorkflow& InWorkflow);
 	void OnCreationWorkflowEnded();
+	
+	void OnSelectionChanged(TSharedPtr<FSandboxListItem> InItem, ESelectInfo::Type) const;
 };
 }
 
diff --git a/Engine/Plugins/Developer/Sandbox/SandboxedEditing/Source/SandboxedEditing/Private/Features/Browser/Widgets/Unsandboxed/SBrowserSandboxDetails.cpp b/Engine/Plugins/Developer/Sandbox/SandboxedEditing/Source/SandboxedEditing/Private/Features/Browser/Widgets/Unsandboxed/SBrowserSandboxDetails.cpp
--- a/Engine/Plugins/Developer/Sandbox/SandboxedEditing/Source/SandboxedEditing/Private/Features/Browser/Widgets/Unsandboxed/SBrowserSandboxDetails.cpp
+++ b/Engine/Plugins/Developer/Sandbox/SandboxedEditing/Source/SandboxedEditing/Private/Features/Browser/Widgets/Unsandboxed/SBrowserSandboxDetails.cpp
@@ -2,6 +2,10 @@
 
 #include "SBrowserSandboxDetails.h"
 
+#include "Components/VerticalBox.h"
+#include "Features/Browser/ViewModels/FileState/FileStateColumnRegistry.h"
+#include "Features/Browser/ViewModels/FileState/UnloadedSandboxFileStateViewModel.h"
+#include "Features/Browser/Widgets/Sandboxed/FileState/SFileStateListView.h"
 #include "Features/Browser/Widgets/Shared/SSandboxDescription.h"
 #include "Internationalization/Text.h"
 #include "Widgets/Layout/SBox.h"
@@ -28,25 +32,23 @@ static TSharedRef<SWidget> MakeMessageWidget(TAttribute<FText> InMessageAttr)
 		];
 }
 }
-void SBrowserSandboxDetails::Construct(const FArguments& InArgs, const TSharedRef<FSandboxMetaDataViewModel>& InMetaDataViewModel)
+void SBrowserSandboxDetails::Construct(
+	const FArguments& InArgs, 
+	const TSharedRef<FSandboxMetaDataViewModel>& InMetaDataViewModel,
+	const TSharedRef<FUnloadedSandboxFileStateViewModel>& InFileActionsViewModel
+	)
 {
 	NumSelectedAttr = InArgs._NumSelected;
 	ChildSlot
 	[
-		SNew(SExpandableArea)
-		.BodyBorderImage(FAppStyle::Get().GetBrush("ToolPanel.GroupBorder"))
-		.HeaderContent()
-		[
-			SNew(STextBlock)
-			.Text(LOCTEXT("Details.Label", "Details"))
-			.ToolTipText(LOCTEXT("Details.ToolTip", "Shows more info about the selected sandbox"))
-		]
-		.BodyContent() [ MakeWidgetSwitcherContent(InArgs, InMetaDataViewModel) ]
+		MakeWidgetSwitcherContent(InArgs, InMetaDataViewModel, InFileActionsViewModel)
 	];
 }
 
 TSharedRef<SWidget> SBrowserSandboxDetails::MakeWidgetSwitcherContent(
-	const FArguments& InArgs, const TSharedRef<FSandboxMetaDataViewModel>& InMetaDataViewModel
+	const FArguments& InArgs, 
+	const TSharedRef<FSandboxMetaDataViewModel>& InMetaDataViewModel, 
+	const TSharedRef<FUnloadedSandboxFileStateViewModel>& InFileActionsViewModel
 	)
 {
 	return SNew(SBox)
@@ -57,8 +59,7 @@ TSharedRef<SWidget> SBrowserSandboxDetails::MakeWidgetSwitcherContent(
 
 			+SWidgetSwitcher::Slot()
 			[
-				SNew(SSandboxDescription, InMetaDataViewModel)
-				.SandboxPath(InArgs._SandboxPath)
+				MakeDetailContent(InArgs, InMetaDataViewModel, InFileActionsViewModel)
 			]
 			
 			+SWidgetSwitcher::Slot()
@@ -73,6 +74,56 @@ TSharedRef<SWidget> SBrowserSandboxDetails::MakeWidgetSwitcherContent(
 		];
 }
 
+TSharedRef<SWidget> SBrowserSandboxDetails::MakeDetailContent(
+	const FArguments& InArgs,
+	const TSharedRef<FSandboxMetaDataViewModel>& InMetaDataViewModel, 
+	const TSharedRef<FUnloadedSandboxFileStateViewModel>& InFileActionsViewModel
+	)
+{
+	return SNew(SVerticalBox)
+		
+		+SVerticalBox::Slot()
+		.AutoHeight()
+		.Padding(5.f)
+		[
+			SNew(SExpandableArea)
+			.BodyBorderImage(FAppStyle::Get().GetBrush("ToolPanel.GroupBorder"))
+			.HeaderContent()
+			[
+				SNew(STextBlock)
+				.Text(LOCTEXT("Description.Label", "Description"))
+				.ToolTipText(LOCTEXT("Description.ToolTip", "Allows you to edit the sandbox description"))
+			]
+			.BodyContent()
+			[
+				SNew(SSandboxDescription, InMetaDataViewModel)
+				.SandboxPath(InArgs._SandboxPath)
+			]
+		]
+		
+		+SVerticalBox::Slot()
+		.FillHeight(1.f)
+		.Padding(5.f)
+		[
+			SNew(SExpandableArea)
+			.BodyBorderImage(FAppStyle::Get().GetBrush("ToolPanel.GroupBorder"))
+			.HeaderContent()
+			[
+				SNew(STextBlock)
+				.Text(LOCTEXT("Content.Label", "File Changes"))
+				.ToolTipText(LOCTEXT("Content.ToolTip", "Shows the files changed by this sandbox"))
+			]
+			.BodyContent()
+			[
+				SNew(SBorder)
+				[
+					SNew(SFileStateListView, InFileActionsViewModel)
+					.ColumnFactories(InArgs._ColumnFactories)
+				]
+			]
+		];
+}
+
 int32 SBrowserSandboxDetails::GetWidgetIndex() const
 {
 	constexpr int32 MainContent = 0;
diff --git a/Engine/Plugins/Developer/Sandbox/SandboxedEditing/Source/SandboxedEditing/Private/Features/Browser/Widgets/Unsandboxed/SBrowserSandboxDetails.h b/Engine/Plugins/Developer/Sandbox/SandboxedEditing/Source/SandboxedEditing/Private/Features/Browser/Widgets/Unsandboxed/SBrowserSandboxDetails.h
--- a/Engine/Plugins/Developer/Sandbox/SandboxedEditing/Source/SandboxedEditing/Private/Features/Browser/Widgets/Unsandboxed/SBrowserSandboxDetails.h
+++ b/Engine/Plugins/Developer/Sandbox/SandboxedEditing/Source/SandboxedEditing/Private/Features/Browser/Widgets/Unsandboxed/SBrowserSandboxDetails.h
@@ -2,11 +2,13 @@
 
 #pragma once
 
+#include "Features/Browser/ViewModels/FileState/FileStateColumns.h"
 #include "Widgets/DeclarativeSyntaxSupport.h"
 #include "Widgets/SCompoundWidget.h"
 
 namespace UE::SandboxedEditing
 {
+class FUnloadedSandboxFileStateViewModel;
 class FSandboxMetaDataViewModel;
 
 /** Shows details for a sandbox in the browser widget. Handles showing text such as "select sandbox to view details", etc. */
@@ -19,9 +21,15 @@ public:
 		SLATE_ATTRIBUTE(TOptional<FString>, SandboxPath)
 		/** The number of selected sandboxes. Relevant for the content displayed. */
 		SLATE_ATTRIBUTE(int32, NumSelected)
+		
+		/** Used to build the columns for the file actions list. */
+		SLATE_ARGUMENT(FFileStateColumnFactoryMap, ColumnFactories)
 	SLATE_END_ARGS()
 
-	void Construct(const FArguments& InArgs, const TSharedRef<FSandboxMetaDataViewModel>& InMetaDataViewModel);
+	void Construct(const FArguments& InArgs, 
+		const TSharedRef<FSandboxMetaDataViewModel>& InMetaDataViewModel,
+		const TSharedRef<FUnloadedSandboxFileStateViewModel>& InFileActionsViewModel
+		);
 	
 private:
 	
@@ -29,7 +37,17 @@ private:
 	TAttribute<int32> NumSelectedAttr;
 	
 	/** @return The widget switcher that switches between the different messages and actual details content. */
-	TSharedRef<SWidget> MakeWidgetSwitcherContent(const FArguments& InArgs, const TSharedRef<FSandboxMetaDataViewModel>& InMetaDataViewModel);
+	TSharedRef<SWidget> MakeWidgetSwitcherContent(
+		const FArguments& InArgs, 
+		const TSharedRef<FSandboxMetaDataViewModel>& InMetaDataViewModel,
+		const TSharedRef<FUnloadedSandboxFileStateViewModel>& InFileActionsViewModel
+	);
+	/** @return Content to display in the widget switcher */
+	TSharedRef<SWidget> MakeDetailContent(
+		const FArguments& InArgs, 
+		const TSharedRef<FSandboxMetaDataViewModel>& InMetaDataViewModel, 
+		const TSharedRef<FUnloadedSandboxFileStateViewModel>& InFileActionsViewModel
+		);
 	
 	/** @return Index for widget switcher to display. Determined by how many items are selected. */
 	int32 GetWidgetIndex() const;
diff --git a/Engine/Plugins/Developer/Sandbox/SandboxedEditing/Source/SandboxedEditing/Private/Features/Browser/Widgets/Unsandboxed/SSandboxedEditingBrowser.cpp b/Engine/Plugins/Developer/Sandbox/SandboxedEditing/Source/SandboxedEditing/Private/Features/Browser/Widgets/Unsandboxed/SSandboxedEditingBrowser.cpp
--- a/Engine/Plugins/Developer/Sandbox/SandboxedEditing/Source/SandboxedEditing/Private/Features/Browser/Widgets/Unsandboxed/SSandboxedEditingBrowser.cpp
+++ b/Engine/Plugins/Developer/Sandbox/SandboxedEditing/Source/SandboxedEditing/Private/Features/Browser/Widgets/Unsandboxed/SSandboxedEditingBrowser.cpp
@@ -10,6 +10,7 @@
 #include "SBrowserControls.h"
 #include "SBrowserFilters.h"
 #include "SBrowserSandboxDetails.h"
+#include "Features/Browser/ViewModels/FileState/UnloadedSandboxFileStateViewModel.h"
 #include "Features/Browser/ViewModels/List/FilterSandboxViewModel.h"
 #include "Utils/BreakBehavior.h"
 #include "Widgets/Text/STextBlock.h"
@@ -21,12 +22,14 @@ namespace UE::SandboxedEditing
 void SSandboxedEditingBrowser::Construct(const FArguments& InArgs, const FViewModels& InViewModels)
 {
 	ControlsViewModel = InViewModels.ControlsViewModel;
+	UnloadedSandboxFileStateViewModel = InViewModels.UnloadedSandboxFileStateViewModel;
 	CommandList = InArgs._CommandList;
 	
 	const SSandboxListView::FViewModels ListViewModels { InViewModels.ListViewModel, InViewModels.ControlsViewModel };
 	const TSharedRef<SSandboxListView> ListViewRef = SNew(SSandboxListView, ListViewModels)
-		.ColumnFactories(InArgs._ColumnFactories)
-		.HighlightText(InViewModels.FilterSandboxViewModel, &FFilterSandboxViewModel::GetSearchText);
+		.ColumnFactories(InArgs._SandboxColumnFactories)
+		.HighlightText(InViewModels.FilterSandboxViewModel, &FFilterSandboxViewModel::GetSearchText)
+		.OnSelectedSandboxChanged(this, &SSandboxedEditingBrowser::OnSelectedSandboxChanged);
 	ListView = ListViewRef;
 	
 	ChildSlot
@@ -37,9 +40,10 @@ void SSandboxedEditingBrowser::Construct(const FArguments& InArgs, const FViewMo
 		.SessionContent() [ ListViewRef ]
 		.DetailsContent()
 		[
-			SNew(SBrowserSandboxDetails, InViewModels.MetaDataViewModel)
+			SNew(SBrowserSandboxDetails, InViewModels.MetaDataViewModel, InViewModels.UnloadedSandboxFileStateViewModel)
 			.SandboxPath(ListViewRef, &SSandboxListView::GetSelectedItem)
 			.NumSelected(ListViewRef, &SSandboxListView::NumSelected)
+			.ColumnFactories(InArgs._FileActionsColumnFactories)
 		]
 	];
 	
@@ -99,6 +103,18 @@ bool SSandboxedEditingBrowser::CanDeleteSelection() const
 	
 	return bCanDelete;
 }
+
+void SSandboxedEditingBrowser::OnSelectedSandboxChanged(TOptional<FString> InSandbox) const
+{
+	if (InSandbox)
+	{
+		UnloadedSandboxFileStateViewModel->SetContent(*InSandbox);
+	}
+	else
+	{
+		UnloadedSandboxFileStateViewModel->ClearContent();
+	}
+}
 }
 
 #undef LOCTEXT_NAMESPACE
\ No newline at end of file
diff --git a/Engine/Plugins/Developer/Sandbox/SandboxedEditing/Source/SandboxedEditing/Private/Features/Browser/Widgets/Unsandboxed/SSandboxedEditingBrowser.h b/Engine/Plugins/Developer/Sandbox/SandboxedEditing/Source/SandboxedEditing/Private/Features/Browser/Widgets/Unsandboxed/SSandboxedEditingBrowser.h
--- a/Engine/Plugins/Developer/Sandbox/SandboxedEditing/Source/SandboxedEditing/Private/Features/Browser/Widgets/Unsandboxed/SSandboxedEditingBrowser.h
+++ b/Engine/Plugins/Developer/Sandbox/SandboxedEditing/Source/SandboxedEditing/Private/Features/Browser/Widgets/Unsandboxed/SSandboxedEditingBrowser.h
@@ -2,6 +2,7 @@
 
 #pragma once
 
+#include "Features/Browser/ViewModels/FileState/FileStateColumns.h"
 #include "Features/Browser/ViewModels/List/SandboxColumns.h"
 #include "Widgets/DeclarativeSyntaxSupport.h"
 #include "Widgets/SCompoundWidget.h"
@@ -15,6 +16,7 @@ class FFilterSandboxViewModel;
 class FSandboxControlsViewModel;
 class FSandboxListViewModel;
 class FSandboxMetaDataViewModel;
+class FUnloadedSandboxFileStateViewModel;
 class SSandboxListView;
 
 /** Displays a list of sandboxes. Shows content only */
@@ -28,13 +30,16 @@ public:
 		TSharedRef<FSandboxListViewModel> ListViewModel;
 		TSharedRef<FSandboxControlsViewModel> ControlsViewModel;
 		TSharedRef<FSandboxMetaDataViewModel> MetaDataViewModel;
+		TSharedRef<FUnloadedSandboxFileStateViewModel> UnloadedSandboxFileStateViewModel;
 	};
 
 	SLATE_BEGIN_ARGS(SSandboxedEditingBrowser){}
 		/** The command list to bind selection-based commands to. */
 		SLATE_ARGUMENT(TSharedPtr<FUICommandList>, CommandList)
 		/** Used to build the columns. */
-		SLATE_ARGUMENT(FSandboxColumnFactoryMap, ColumnFactories)
+		SLATE_ARGUMENT(FSandboxColumnFactoryMap, SandboxColumnFactories)
+		/** Used to build the columns for the file actions list. */
+		SLATE_ARGUMENT(FFileStateColumnFactoryMap, FileActionsColumnFactories)
 	SLATE_END_ARGS()
 
 	void Construct(const FArguments& InArgs, const FViewModels& InViewModels);
@@ -43,6 +48,9 @@ private:
 	
 	/** Used for implementing the commands we bind to. */
 	TSharedPtr<FSandboxControlsViewModel> ControlsViewModel;
+	/** When the selected sandbox changes, we notify this model. */
+	TSharedPtr<FUnloadedSandboxFileStateViewModel> UnloadedSandboxFileStateViewModel;
+	
 	/** Optional command list for selection based commands, such as delete. */
 	TSharedPtr<FUICommandList> CommandList;
 	
@@ -57,6 +65,9 @@ private:
 	void HandleDeleteSelection() const;
 	/** @return Whether the current selection can be deleted. */
 	bool CanDeleteSelection() const;
+	
+	/** Notify the UnloadedSandboxFileStateViewModel when the selected sandbox changes. */
+	void OnSelectedSandboxChanged(TOptional<FString> InSandbox) const;
 };
 }
 
diff --git a/Engine/Plugins/Developer/Sandbox/SandboxedEditing/Source/SandboxedEditing/Private/Framework/Models/SandboxSystemModel.cpp b/Engine/Plugins/Developer/Sandbox/SandboxedEditing/Source/SandboxedEditing/Private/Framework/Models/SandboxSystemModel.cpp
--- a/Engine/Plugins/Developer/Sandbox/SandboxedEditing/Source/SandboxedEditing/Private/Framework/Models/SandboxSystemModel.cpp
+++ b/Engine/Plugins/Developer/Sandbox/SandboxedEditing/Source/SandboxedEditing/Private/Framework/Models/SandboxSystemModel.cpp
@@ -250,7 +250,7 @@ bool FSandboxSystemModel::CanPersistAllChanges() const
 	return HasActiveSandbox();
 }
 
-FileSandboxCore::FGatheredFileChanges FSandboxSystemModel::GatherFileChanges() const
+FileSandboxCore::FGatheredFileChanges FSandboxSystemModel::GatherActiveFileChanges() const
 {
 	using namespace FileSandboxCore;
 	IFileSandboxCoreModule& Module = IFileSandboxCoreModule::Get();
@@ -259,6 +259,22 @@ FileSandboxCore::FGatheredFileChanges FSandboxSystemModel::GatherFileChanges() c
 	return Instance ? Instance->GatherChangedFiles() : FGatheredFileChanges{};
 }
 
+FileSandboxCore::FGatheredFileChanges FSandboxSystemModel::GatherFileChanges(const FString& InSandbox) const
+{
+	using namespace FileSandboxCore;
+	IFileSandboxCoreModule& Module = IFileSandboxCoreModule::Get();
+	ISandboxManager& SandboxManager = Module.GetSandboxManager();
+	
+	FGatheredFileChanges Result;
+	SandboxManager.EnumerateFileChanges(InSandbox, [&Result](const FString& InPath, ESandboxFileChange InChange)
+	{
+		Result.NonSandboxPaths.Add(InPath);
+		Result.FileActions.Add(InChange);
+		return FileSandboxCore::EBreakBehavior::Continue;
+	});
+	return Result;
+}
+
 void FSandboxSystemModel::DiscardAllChanges() const
 {
 	using namespace FileSandboxCore;
diff --git a/Engine/Plugins/Developer/Sandbox/SandboxedEditing/Source/SandboxedEditing/Private/Framework/Models/SandboxSystemModel.h b/Engine/Plugins/Developer/Sandbox/SandboxedEditing/Source/SandboxedEditing/Private/Framework/Models/SandboxSystemModel.h
--- a/Engine/Plugins/Developer/Sandbox/SandboxedEditing/Source/SandboxedEditing/Private/Framework/Models/SandboxSystemModel.h
+++ b/Engine/Plugins/Developer/Sandbox/SandboxedEditing/Source/SandboxedEditing/Private/Framework/Models/SandboxSystemModel.h
@@ -74,7 +74,9 @@ public:
 	bool CanPersistAllChanges() const;
 	
 	/** @return The file changes that were performed by the sandbox. */
-	FileSandboxCore::FGatheredFileChanges GatherFileChanges() const;
+	FileSandboxCore::FGatheredFileChanges GatherActiveFileChanges() const;
+	/** @return The file changes of the specified sandbox. */
+	FileSandboxCore::FGatheredFileChanges GatherFileChanges(const FString& InSandbox) const;
 
 	/** Discards all changes */
 	void DiscardAllChanges() const;
