Sandboxed Editing: Add persist button

#Synced-CL 49159615

diff --git a/Engine/Plugins/Developer/Sandbox/SandboxedEditing/Source/SandboxedEditing/Private/Features/Browser/Commands/BrowserCommandBindings.cpp b/Engine/Plugins/Developer/Sandbox/SandboxedEditing/Source/SandboxedEditing/Private/Features/Browser/Commands/BrowserCommandBindings.cpp
--- a/Engine/Plugins/Developer/Sandbox/SandboxedEditing/Source/SandboxedEditing/Private/Features/Browser/Commands/BrowserCommandBindings.cpp
+++ b/Engine/Plugins/Developer/Sandbox/SandboxedEditing/Source/SandboxedEditing/Private/Features/Browser/Commands/BrowserCommandBindings.cpp
@@ -13,6 +13,7 @@ namespace UE::SandboxedEditing
 FBrowserCommandBindings::FBrowserCommandBindings(const FBrowserViewModels& InViewModels)
 	: ControlsViewModel(InViewModels.ControlsViewModel)
 	, LeaveSandboxViewModel(InViewModels.LeaveViewModel)
+	, PersistViewModel(InViewModels.PersistViewModel)
 	, CommandList(MakeShared<FUICommandList>())
 {
 	FBrowserCommands::Register();
@@ -27,6 +28,10 @@ FBrowserCommandBindings::FBrowserCommandBindings(const FBrowserViewModels& InVie
 		FExecuteAction::CreateRaw(this, &FBrowserCommandBindings::HandleLeaveSandbox),
 		FCanExecuteAction::CreateRaw(this, &FBrowserCommandBindings::CanLeaveSandbox)
 		);
+	CommandList->MapAction(BrowserCommands.PersistSandbox,
+		FExecuteAction::CreateRaw(this, &FBrowserCommandBindings::HandlePersistSandbox),
+		FCanExecuteAction::CreateRaw(this, &FBrowserCommandBindings::CanPersistSandbox)
+		);
 }
 
 FBrowserCommandBindings::~FBrowserCommandBindings()
@@ -54,4 +59,14 @@ bool FBrowserCommandBindings::CanLeaveSandbox()
 {
 	return LeaveSandboxViewModel->CanLeaveSandbox();
 }
+
+void FBrowserCommandBindings::HandlePersistSandbox()
+{
+	PersistViewModel->StartPersistWorkflowForActiveSandbox();
+}
+
+bool FBrowserCommandBindings::CanPersistSandbox()
+{
+	return PersistViewModel->CanStartPersistWorkflowForActiveSandbox();
+}
 }
diff --git a/Engine/Plugins/Developer/Sandbox/SandboxedEditing/Source/SandboxedEditing/Private/Features/Browser/Commands/BrowserCommandBindings.h b/Engine/Plugins/Developer/Sandbox/SandboxedEditing/Source/SandboxedEditing/Private/Features/Browser/Commands/BrowserCommandBindings.h
--- a/Engine/Plugins/Developer/Sandbox/SandboxedEditing/Source/SandboxedEditing/Private/Features/Browser/Commands/BrowserCommandBindings.h
+++ b/Engine/Plugins/Developer/Sandbox/SandboxedEditing/Source/SandboxedEditing/Private/Features/Browser/Commands/BrowserCommandBindings.h
@@ -8,8 +8,9 @@
 
 namespace UE::SandboxedEditing
 {
-class FLeaveSandboxViewModel;
 class FActiveSandboxDetailsViewModel;
+class FLeaveSandboxViewModel;
+class FPersistSandboxViewModel;
 class FSandboxControlsViewModel;
 struct FBrowserViewModels;
 
@@ -29,6 +30,8 @@ private:
 	const TSharedRef<FSandboxControlsViewModel> ControlsViewModel;
 	/** Used to interact with the active sandbox. */
 	const TSharedRef<FLeaveSandboxViewModel> LeaveSandboxViewModel;
+	/** Used to persist the sandbox. */
+	const TSharedRef<FPersistSandboxViewModel> PersistViewModel;
 	
 	/** Holds our command bindings */
 	const TSharedRef<FUICommandList> CommandList;
@@ -41,6 +44,9 @@ private:
 	
 	void HandleLeaveSandbox();
 	bool CanLeaveSandbox();
+	
+	void HandlePersistSandbox();
+	bool CanPersistSandbox();
 };
 }
 
diff --git a/Engine/Plugins/Developer/Sandbox/SandboxedEditing/Source/SandboxedEditing/Private/Features/Browser/Commands/BrowserCommands.cpp b/Engine/Plugins/Developer/Sandbox/SandboxedEditing/Source/SandboxedEditing/Private/Features/Browser/Commands/BrowserCommands.cpp
--- a/Engine/Plugins/Developer/Sandbox/SandboxedEditing/Source/SandboxedEditing/Private/Features/Browser/Commands/BrowserCommands.cpp
+++ b/Engine/Plugins/Developer/Sandbox/SandboxedEditing/Source/SandboxedEditing/Private/Features/Browser/Commands/BrowserCommands.cpp
@@ -21,8 +21,9 @@ void FBrowserCommands::RegisterCommands()
 {
 	UI_COMMAND(CreateNewSandbox, "Create New Sandbox", "Creates a new sandbox", EUserInterfaceActionType::Button, FInputChord());
 	UI_COMMAND(Cancel, "Cancel", "Cancel the current operation, such as renaming or creating a sandbox.", EUserInterfaceActionType::None, FInputChord(EKeys::Escape));
-	UI_COMMAND(LeaveSandbox, "Leave Sandbox", "Leaves the current sandbox.", EUserInterfaceActionType::Button, FInputChord());
 	
+	UI_COMMAND(LeaveSandbox, "Leave Sandbox", "Leaves the current sandbox.", EUserInterfaceActionType::Button, FInputChord());
+	UI_COMMAND(PersistSandbox, "Persist Sandbox", "Persists the active sandbox.", EUserInterfaceActionType::Button, FInputChord());
 }
 }
 
diff --git a/Engine/Plugins/Developer/Sandbox/SandboxedEditing/Source/SandboxedEditing/Private/Features/Browser/Commands/BrowserCommands.h b/Engine/Plugins/Developer/Sandbox/SandboxedEditing/Source/SandboxedEditing/Private/Features/Browser/Commands/BrowserCommands.h
--- a/Engine/Plugins/Developer/Sandbox/SandboxedEditing/Source/SandboxedEditing/Private/Features/Browser/Commands/BrowserCommands.h
+++ b/Engine/Plugins/Developer/Sandbox/SandboxedEditing/Source/SandboxedEditing/Private/Features/Browser/Commands/BrowserCommands.h
@@ -25,6 +25,8 @@ public:
 	
 	/** Leaves the current sandbox. */
 	TSharedPtr<FUICommandInfo> LeaveSandbox;
+	/** Persists the current sandbox (without leaving). */
+	TSharedPtr<FUICommandInfo> PersistSandbox;
 };
 }
 
diff --git a/Engine/Plugins/Developer/Sandbox/SandboxedEditing/Source/SandboxedEditing/Private/Features/Browser/ViewModels/Persist/PersistSandboxViewModel.cpp b/Engine/Plugins/Developer/Sandbox/SandboxedEditing/Source/SandboxedEditing/Private/Features/Browser/ViewModels/Persist/PersistSandboxViewModel.cpp
--- a/Engine/Plugins/Developer/Sandbox/SandboxedEditing/Source/SandboxedEditing/Private/Features/Browser/ViewModels/Persist/PersistSandboxViewModel.cpp
+++ b/Engine/Plugins/Developer/Sandbox/SandboxedEditing/Source/SandboxedEditing/Private/Features/Browser/ViewModels/Persist/PersistSandboxViewModel.cpp
@@ -27,11 +27,22 @@ bool FPersistSandboxViewModel::StartPersistWorkflow(FileSandboxCore::FGatheredFi
 	return false;
 }
 
+bool FPersistSandboxViewModel::StartPersistWorkflowForActiveSandbox()
+{
+	FileSandboxCore::FGatheredFileChanges FileChanges = Model->GatherActiveFileChanges();
+	return FileChanges.HasChanges() && StartPersistWorkflow(MoveTemp(FileChanges));
+}
+
 bool FPersistSandboxViewModel::CanStartPersistWorkflow() const
 {
 	return !ActivePersistWorkflow.IsValid();
 }
 
+bool FPersistSandboxViewModel::CanStartPersistWorkflowForActiveSandbox() const
+{
+	return CanStartPersistWorkflow() && Model->GatherActiveFileChanges().HasChanges();
+}
+
 void FPersistSandboxViewModel::OnCleanupPersistWorkflow(const FPersistWorkflowUserChoice& InPersistResult)
 {
 	if (InPersistResult.UserAction == EPersistWorkflowAction::Persist)
diff --git a/Engine/Plugins/Developer/Sandbox/SandboxedEditing/Source/SandboxedEditing/Private/Features/Browser/ViewModels/Persist/PersistSandboxViewModel.h b/Engine/Plugins/Developer/Sandbox/SandboxedEditing/Source/SandboxedEditing/Private/Features/Browser/ViewModels/Persist/PersistSandboxViewModel.h
--- a/Engine/Plugins/Developer/Sandbox/SandboxedEditing/Source/SandboxedEditing/Private/Features/Browser/ViewModels/Persist/PersistSandboxViewModel.h
+++ b/Engine/Plugins/Developer/Sandbox/SandboxedEditing/Source/SandboxedEditing/Private/Features/Browser/ViewModels/Persist/PersistSandboxViewModel.h
@@ -24,9 +24,13 @@ public:
 
 	/** Starts the persist workflow. */
 	bool StartPersistWorkflow(FileSandboxCore::FGatheredFileChanges InFileChanges);
+	/** Starts a perist workflow by using the changes of the active sandbox. */
+	bool StartPersistWorkflowForActiveSandbox();
 	
 	/** @return Whether a persist workflow can be started */
 	bool CanStartPersistWorkflow() const;
+	/** @return Whether a persist workflow can be started for the active sandbox */
+	bool CanStartPersistWorkflowForActiveSandbox() const;
 	
 	/** @return The current persist workflow */
 	FPersistSandboxWorkflow* GetPersistWorkflow() const { return ActivePersistWorkflow.Get(); }
diff --git a/Engine/Plugins/Developer/Sandbox/SandboxedEditing/Source/SandboxedEditing/Private/Features/Browser/Widgets/Sandboxed/SActiveSandboxToolbar.cpp b/Engine/Plugins/Developer/Sandbox/SandboxedEditing/Source/SandboxedEditing/Private/Features/Browser/Widgets/Sandboxed/SActiveSandboxToolbar.cpp
--- a/Engine/Plugins/Developer/Sandbox/SandboxedEditing/Source/SandboxedEditing/Private/Features/Browser/Widgets/Sandboxed/SActiveSandboxToolbar.cpp
+++ b/Engine/Plugins/Developer/Sandbox/SandboxedEditing/Source/SandboxedEditing/Private/Features/Browser/Widgets/Sandboxed/SActiveSandboxToolbar.cpp
@@ -5,6 +5,7 @@
 #include "Features/Browser/Commands/BrowserCommands.h"
 #include "Features/Browser/ViewModels/Active/ActiveSandboxDetailsViewModel.h"
 #include "SandboxedEditingStyle.h"
+#include "SPrimaryButton.h"
 #include "Styling/AppStyle.h"
 #include "Utils/WidgetUtils.h"
 #include "Widgets/Layout/SBorder.h"
@@ -13,6 +14,8 @@
 #include "Widgets/SBoxPanel.h"
 #include "Widgets/Text/STextBlock.h"
 
+#define LOCTEXT_NAMESPACE "SActiveSandboxToolbar"
+
 namespace UE::SandboxedEditing
 {
 void SActiveSandboxToolbar::Construct(const FArguments& InArgs, const TSharedRef<FActiveSandboxDetailsViewModel>& InActiveSandboxViewModel)
@@ -38,6 +41,17 @@ void SActiveSandboxToolbar::Construct(const FArguments& InArgs, const TSharedRef
 			SNew(SSeparator)
 			.Orientation(Orient_Vertical)
 		]
+		
+		+SHorizontalBox::Slot()
+		.AutoWidth()
+		.VAlign(VAlign_Center)
+		.Padding(FMargin(1.0f))
+		[
+			MakePrimaryButtonByCommand(
+				InArgs._CommandList, FBrowserCommands::Get().PersistSandbox,
+				LOCTEXT("PersistLabel", "Persist")
+				)
+		]
 
 		// Move all the other content to the right if the panel is big
 		+SHorizontalBox::Slot()
@@ -82,3 +96,5 @@ FText SActiveSandboxToolbar::GetLeaveButtonToolTipText() const
 		: Reason;
 }
 }
+
+#undef LOCTEXT_NAMESPACE
\ No newline at end of file
diff --git a/Engine/Plugins/Developer/Sandbox/SandboxedEditing/Source/SandboxedEditing/Private/Utils/WidgetUtils.cpp b/Engine/Plugins/Developer/Sandbox/SandboxedEditing/Source/SandboxedEditing/Private/Utils/WidgetUtils.cpp
--- a/Engine/Plugins/Developer/Sandbox/SandboxedEditing/Source/SandboxedEditing/Private/Utils/WidgetUtils.cpp
+++ b/Engine/Plugins/Developer/Sandbox/SandboxedEditing/Source/SandboxedEditing/Private/Utils/WidgetUtils.cpp
@@ -4,6 +4,7 @@
 
 #include "SNegativeActionButton.h"
 #include "SPositiveActionButton.h"
+#include "SPrimaryButton.h"
 #include "Framework/SlateDelegates.h"
 #include "Framework/Commands/UICommandInfo.h"
 #include "Framework/Commands/UICommandList.h"
@@ -39,6 +40,24 @@ TSharedRef<SWidget> MakeIconButton(
 		];
 }
 
+TSharedRef<SWidget> MakePrimaryButton(
+	const TAttribute<FText>& Label,
+	const TAttribute<FText>& Tooltip,
+	const TAttribute<const FSlateBrush*>& Icon,
+	const TAttribute<bool>& EnabledAttribute,
+	const FOnClicked& OnClicked,
+	const TAttribute<EVisibility>& Visibility
+)
+{
+	return SNew(SPrimaryButton)
+		.Icon(Icon)
+		.ToolTipText(Tooltip)
+		.IsEnabled(EnabledAttribute)
+		.OnClicked(OnClicked)
+		.Visibility(Visibility)
+		.Text(Label);
+}
+
 TSharedRef<SWidget> MakeIconButtonByCommand(
 	const TAttribute<const FSlateBrush*>& Icon,
 	const TSharedPtr<FUICommandList>& InCommandList,
@@ -63,8 +82,38 @@ TSharedRef<SWidget> MakeIconButtonByCommand(
 			}
 			return FReply::Handled();
 		}), 
-		Visibility)
-	;
+		Visibility
+		);
+}
+
+TSharedRef<SWidget> MakePrimaryButtonByCommand(
+	const TSharedPtr<FUICommandList>& InCommandList,
+	const TSharedPtr<FUICommandInfo>& InCommand,
+	const TAttribute<FText>& InLabel,
+	const TAttribute<FText>& InOverrideTooltip,
+	const TAttribute<const FSlateBrush*>& InIcon,
+	const TAttribute<EVisibility>& Visibility
+)
+{
+	return MakePrimaryButton(
+		InLabel, 
+		InOverrideTooltip.IsBound() || InOverrideTooltip.IsSet() ? InOverrideTooltip : InCommand->GetDescription(),
+		InIcon, 
+		TAttribute<bool>::CreateLambda([InCommandList, InCommand]
+			{
+				const FUIAction* Action = InCommandList->GetActionForCommand(InCommand);
+				return Action && Action->CanExecute();
+			}),
+		FOnClicked::CreateLambda([InCommandList, InCommand]
+			{
+				if (const FUIAction* Action = InCommandList->GetActionForCommand(InCommand))
+				{
+					Action->Execute();
+				}
+				return FReply::Handled();
+			}), 
+		Visibility
+	);
 }
 
 TSharedRef<SWidget> MakePositiveActionButton(
diff --git a/Engine/Plugins/Developer/Sandbox/SandboxedEditing/Source/SandboxedEditing/Private/Utils/WidgetUtils.h b/Engine/Plugins/Developer/Sandbox/SandboxedEditing/Source/SandboxedEditing/Private/Utils/WidgetUtils.h
--- a/Engine/Plugins/Developer/Sandbox/SandboxedEditing/Source/SandboxedEditing/Private/Utils/WidgetUtils.h
+++ b/Engine/Plugins/Developer/Sandbox/SandboxedEditing/Source/SandboxedEditing/Private/Utils/WidgetUtils.h
@@ -23,6 +23,18 @@ TSharedRef<SWidget> MakeIconButton(
 	const FOnClicked& OnClicked, 
 	const TAttribute<EVisibility>& Visibility = EVisibility::Visible
 	);
+
+/** Makes primary action button with a label. */
+TSharedRef<SWidget> MakePrimaryButton(
+	const TAttribute<FText>& Label,
+	const TAttribute<FText>& Tooltip,
+	const TAttribute<const FSlateBrush*>& Icon,
+	const TAttribute<bool>& EnabledAttribute,
+	const FOnClicked& OnClicked,
+	const TAttribute<EVisibility>& Visibility = EVisibility::Visible
+);
+
+
 /** Overload that determines the IsEnabled and OnClicked parameters from the command list. */
 TSharedRef<SWidget> MakeIconButtonByCommand(
 	const TAttribute<const FSlateBrush*>& Icon,
@@ -32,6 +44,16 @@ TSharedRef<SWidget> MakeIconButtonByCommand(
 	const TAttribute<EVisibility>& Visibility = EVisibility::Visible
 );
 
+/** Overload that determines the IsEnabled and OnClicked parameters from the command list. */
+TSharedRef<SWidget> MakePrimaryButtonByCommand(
+	const TSharedPtr<FUICommandList>& InCommandList,
+	const TSharedPtr<FUICommandInfo>& InCommand,
+	const TAttribute<FText>& InLabel,
+	const TAttribute<FText>& InOverrideTooltip = {},
+	const TAttribute<const FSlateBrush*>& InIcon = {},
+	const TAttribute<EVisibility>& Visibility = EVisibility::Visible
+);
+
 /** Makes a positive action button that only contains an icon */
 TSharedRef<SWidget> MakePositiveActionButton(
 	const TAttribute<const FSlateBrush*>& Icon, 
