Sandboxed Editing: Show file changes in active sandbox UI.

#Synced-CL 49159615

diff --git a/Engine/Plugins/Developer/Sandbox/FileSandbox/Source/FileSandboxCore/Private/Sandbox/Persist/DiscardUtils.cpp b/Engine/Plugins/Developer/Sandbox/FileSandbox/Source/FileSandboxCore/Private/Sandbox/Persist/DiscardUtils.cpp
--- a/Engine/Plugins/Developer/Sandbox/FileSandbox/Source/FileSandboxCore/Private/Sandbox/Persist/DiscardUtils.cpp
+++ b/Engine/Plugins/Developer/Sandbox/FileSandbox/Source/FileSandboxCore/Private/Sandbox/Persist/DiscardUtils.cpp
@@ -73,7 +73,7 @@ void DiscardAddedAndModifiedPaths(
 			OutPackagesPendingHotReload.Remove(PackageName);
 		}
 				
-		InProcessFileChangeCallback(FFileChange{ InPath.GetNonSandboxPath(), EFileChangeAction::Added });
+		InProcessFileChangeCallback(FFileChange{ InPath.GetNonSandboxPath(), EFileChangeAction::Removed });
 		return EBreakBehavior::Continue;
 	};
 	
diff --git a/Engine/Plugins/Developer/Sandbox/SandboxedEditing/Source/SandboxedEditing/Private/Features/Browser/BrowserFeature.cpp b/Engine/Plugins/Developer/Sandbox/SandboxedEditing/Source/SandboxedEditing/Private/Features/Browser/BrowserFeature.cpp
--- a/Engine/Plugins/Developer/Sandbox/SandboxedEditing/Source/SandboxedEditing/Private/Features/Browser/BrowserFeature.cpp
+++ b/Engine/Plugins/Developer/Sandbox/SandboxedEditing/Source/SandboxedEditing/Private/Features/Browser/BrowserFeature.cpp
@@ -55,8 +55,9 @@ FBrowserFeature::FBrowserFeature(
 	)
 	: SandboxModel(InSandboxModel)
 	, SandboxColumns(BrowserDetail::MakeSandboxColumnRegistry())
-	, BrowserColumns(GetColumnsForBrowser())
-	, ViewModels(SandboxColumns, BrowserColumns, SandboxModel, InExternalSandboxViewModel)
+	, BrowserFileActionColumns(GetColumnsForBrowser())
+	, ActiveFileActionColumns(GetColumnsForActiveSandbox())
+	, ViewModels(SandboxColumns, BrowserFileActionColumns, ActiveFileActionColumns, SandboxModel, InExternalSandboxViewModel)
 	, Views(ViewModels.PersistViewModel, ViewModels.LeaveViewModel)
 	, CommandBindings(MakeShared<FBrowserCommandBindings>(ViewModels))
 {
@@ -101,7 +102,8 @@ TSharedRef<SDockTab> FBrowserFeature::SpawnBrowserTab(const FSpawnTabArgs& InArg
 			SNew(SSandboxedEditingRoot, ViewModels)
 			.CommandList(CommandBindings->GetCommandList())
 			.SandboxColumnFactories(SandboxColumns.ColumnFactories)
-			.FileActionsColumnFactories(BrowserColumns.ColumnFactories)
+			.BrowserFileActionsColumnFactories(BrowserFileActionColumns.ColumnFactories)
+			.ActiveFileStateColumnsFactories(ActiveFileActionColumns.ColumnFactories)
 		];
 }
 }
diff --git a/Engine/Plugins/Developer/Sandbox/SandboxedEditing/Source/SandboxedEditing/Private/Features/Browser/BrowserFeature.h b/Engine/Plugins/Developer/Sandbox/SandboxedEditing/Source/SandboxedEditing/Private/Features/Browser/BrowserFeature.h
--- a/Engine/Plugins/Developer/Sandbox/SandboxedEditing/Source/SandboxedEditing/Private/Features/Browser/BrowserFeature.h
+++ b/Engine/Plugins/Developer/Sandbox/SandboxedEditing/Source/SandboxedEditing/Private/Features/Browser/BrowserFeature.h
@@ -46,8 +46,10 @@ private:
 	
 	/** The columns that are displayed in the browser. */
 	const FSandboxColumnRegistry SandboxColumns;
-	/** The columns displayed for the file actions in the browser. */
-	const FFileStateColumnRegistry BrowserColumns;
+	/** The columns displayed for the file actions in the browser UI. */
+	const FFileStateColumnRegistry BrowserFileActionColumns;
+	/** The columns displayed for the file actions in the active sandbox UI. */
+	const FFileStateColumnRegistry ActiveFileActionColumns;
 	
 	/** Holds all view models relevant for the browser. */
 	const FBrowserViewModels ViewModels;
diff --git a/Engine/Plugins/Developer/Sandbox/SandboxedEditing/Source/SandboxedEditing/Private/Features/Browser/ViewModels/BrowserViewModels.h b/Engine/Plugins/Developer/Sandbox/SandboxedEditing/Source/SandboxedEditing/Private/Features/Browser/ViewModels/BrowserViewModels.h
--- a/Engine/Plugins/Developer/Sandbox/SandboxedEditing/Source/SandboxedEditing/Private/Features/Browser/ViewModels/BrowserViewModels.h
+++ b/Engine/Plugins/Developer/Sandbox/SandboxedEditing/Source/SandboxedEditing/Private/Features/Browser/ViewModels/BrowserViewModels.h
@@ -5,7 +5,8 @@
 #include "Active/ActiveSandboxDetailsViewModel.h"
 #include "Active/ActiveSandboxTrackerViewModel.h"
 #include "FileState/FileStateColumnRegistry.h"
-#include "FileState/UnloadedSandboxFileStateViewModel.h"
+#include "FileState/Models/ActiveSandboxFileStateViewModel.h"
+#include "FileState/Models/UnloadedSandboxFileStateViewModel.h"
 #include "Leaving/LeaveSandboxViewModel.h"
 #include "List/FilterSandboxViewModel.h"
 #include "List/SandboxColumnRegistry.h"
@@ -28,8 +29,11 @@ struct FBrowserViewModels
 	const TSharedRef<FFilterSandboxViewModel> FilterSandboxViewModel;
 	/** Knows about the items that are displayed in the browser. */
 	const TSharedRef<FSandboxListViewModel> ListViewModel;
+	
 	/** Knows the file action of the currently selected sandbox. */
 	const TSharedRef<FUnloadedSandboxFileStateViewModel> UnloadedSandboxFileStateViewModel;
+	/** Knows the file actions of the active sandbox. */
+	const TSharedRef<FActiveSandboxFileStateViewModel> ActiveSandboxFileStateViewModel;
 	
 	/** Handles actions that the user performs on the sandboxes via the UI. */
 	const TSharedRef<FSandboxControlsViewModel> ControlsViewModel;
@@ -51,13 +55,15 @@ struct FBrowserViewModels
 	
 	explicit FBrowserViewModels(
 		const FSandboxColumnRegistry& InSandboxColumns, 
-		const FFileStateColumnRegistry& InBrowserColumns,
+		const FFileStateColumnRegistry& InBrowserFileActionColumns,
+		const FFileStateColumnRegistry& InActiveFileActionColumns,
 		const TSharedRef<FSandboxSystemModel>& InModel,
 		const TSharedRef<FileSandboxUI::IExternalSandboxActiveViewModel>& InExternalSandboxViewModel
 		)
 		: FilterSandboxViewModel(MakeShared<FFilterSandboxViewModel>(InSandboxColumns.GetBehaviorArray()))
 		, ListViewModel(MakeShared<FSandboxListViewModel>(InModel, FilterSandboxViewModel, InSandboxColumns.ColumnBehaviors))
-		, UnloadedSandboxFileStateViewModel(MakeShared<FUnloadedSandboxFileStateViewModel>(InModel, InBrowserColumns.ColumnBehaviors))
+		, UnloadedSandboxFileStateViewModel(MakeShared<FUnloadedSandboxFileStateViewModel>(InModel, InBrowserFileActionColumns.ColumnBehaviors))
+		, ActiveSandboxFileStateViewModel(MakeShared<FActiveSandboxFileStateViewModel>(InModel, InActiveFileActionColumns.ColumnBehaviors))
 		, ControlsViewModel(MakeShared<FSandboxControlsViewModel>(InModel))
 		, MetaDataViewModel(MakeShared<FSandboxMetaDataViewModel>(InModel))
 		, PersistViewModel(MakeShared<FPersistSandboxViewModel>(InModel))
diff --git a/Engine/Plugins/Developer/Sandbox/SandboxedEditing/Source/SandboxedEditing/Private/Features/Browser/ViewModels/FileState/FileStateColumnRegistry.cpp b/Engine/Plugins/Developer/Sandbox/SandboxedEditing/Source/SandboxedEditing/Private/Features/Browser/ViewModels/FileState/FileStateColumnRegistry.cpp
--- a/Engine/Plugins/Developer/Sandbox/SandboxedEditing/Source/SandboxedEditing/Private/Features/Browser/ViewModels/FileState/FileStateColumnRegistry.cpp
+++ b/Engine/Plugins/Developer/Sandbox/SandboxedEditing/Source/SandboxedEditing/Private/Features/Browser/ViewModels/FileState/FileStateColumnRegistry.cpp
@@ -30,6 +30,25 @@ FFileStateColumnRegistry GetColumnsForBrowser()
 	return Registry;
 }
 
+FFileStateColumnRegistry GetColumnsForActiveSandbox()
+{
+	const TSharedRef<FActionFileStateColumn> ActionColum = MakeShared<FActionFileStateColumn>();
+	const TSharedRef<FPathFileStateColumn> PathColumn = MakeShared<FPathFileStateColumn>();
+	
+	FFileStateColumnRegistry Registry;
+	Registry.ColumnBehaviors = TMap<FName, TSharedRef<IFileStateColumnBehavior>>
+	{
+		{ FileActionColumn, ActionColum },
+		{ FilePathColumn, PathColumn }
+	};
+	Registry.ColumnFactories = TMap<FName, TSharedRef<IFileStateColumnWidgetFactory>>
+	{
+		{ FileActionColumn, ActionColum },
+		{ FilePathColumn, PathColumn }
+	};
+	return Registry;
+}
+
 FFileStateColumnRegistry GetColumnsForPersist(const TSharedRef<FPersistOperationViewModel>& InViewModel)
 {
 	const TSharedRef<FPersistFileStateColumn> CheckboxColumn = MakeShared<FPersistFileStateColumn>(InViewModel);
diff --git a/Engine/Plugins/Developer/Sandbox/SandboxedEditing/Source/SandboxedEditing/Private/Features/Browser/ViewModels/FileState/FileStateColumnRegistry.h b/Engine/Plugins/Developer/Sandbox/SandboxedEditing/Source/SandboxedEditing/Private/Features/Browser/ViewModels/FileState/FileStateColumnRegistry.h
--- a/Engine/Plugins/Developer/Sandbox/SandboxedEditing/Source/SandboxedEditing/Private/Features/Browser/ViewModels/FileState/FileStateColumnRegistry.h
+++ b/Engine/Plugins/Developer/Sandbox/SandboxedEditing/Source/SandboxedEditing/Private/Features/Browser/ViewModels/FileState/FileStateColumnRegistry.h
@@ -22,6 +22,10 @@ struct FFileStateColumnRegistry
 	
 /** @return Columns to show in the browser view. */
 FFileStateColumnRegistry GetColumnsForBrowser();
+
+/** @return Columns to show in the active sandbox view. */
+FFileStateColumnRegistry GetColumnsForActiveSandbox();
+
 /** @return Columns to show in the persist view. */
 FFileStateColumnRegistry GetColumnsForPersist(const TSharedRef<FPersistOperationViewModel>& InViewModel);
 }
\ No newline at end of file
diff --git a/Engine/Plugins/Developer/Sandbox/SandboxedEditing/Source/SandboxedEditing/Private/Features/Browser/ViewModels/FileState/Models/ActiveSandboxFileStateViewModel.cpp b/Engine/Plugins/Developer/Sandbox/SandboxedEditing/Source/SandboxedEditing/Private/Features/Browser/ViewModels/FileState/Models/ActiveSandboxFileStateViewModel.cpp
new file mode 100644
--- /dev/null
+++ b/Engine/Plugins/Developer/Sandbox/SandboxedEditing/Source/SandboxedEditing/Private/Features/Browser/ViewModels/FileState/Models/ActiveSandboxFileStateViewModel.cpp
@@ -0,0 +1,90 @@
+﻿// Copyright Epic Games, Inc. All Rights Reserved.
+
+#include "ActiveSandboxFileStateViewModel.h"
+
+#include "Features/Browser/ViewModels/FileState/FileStateItem.h"
+#include "Framework/Models/SandboxSystemModel.h"
+#include "Misc/Paths.h"
+#include "Types/EBreakBehavior.h"
+
+namespace UE::SandboxedEditing
+{
+FActiveSandboxFileStateViewModel::FActiveSandboxFileStateViewModel(
+	const TSharedRef<FSandboxSystemModel>& InModel,
+	TMap<FName, TSharedRef<IFileStateColumnBehavior>> InColumns
+	)
+	: FStaticFileStateViewModel(MoveTemp(InColumns), {})
+	, Model(InModel)
+{
+	Model->OnLoadSandbox().AddRaw(this, &FActiveSandboxFileStateViewModel::OnSandboxLoaded);
+	Model->OnLeaveSandbox().AddRaw(this, &FActiveSandboxFileStateViewModel::OnSandboxLeft);
+	Model->OnSandboxFilesChanged().AddRaw(this, &FActiveSandboxFileStateViewModel::RefreshItems);
+	
+	if (Model->HasActiveSandbox())
+	{
+		OnSandboxLoaded();
+	}
+}
+
+FActiveSandboxFileStateViewModel::~FActiveSandboxFileStateViewModel()
+{
+	Model->OnLoadSandbox().RemoveAll(this);
+	Model->OnLeaveSandbox().RemoveAll(this);
+	Model->OnSandboxFilesChanged().RemoveAll(this);
+}
+
+void FActiveSandboxFileStateViewModel::RefreshItems()
+{
+	if (!Model->HasActiveSandbox() && Items.IsEmpty())
+	{
+		return;
+	}
+	
+	if (!Model->HasActiveSandbox() && !Items.IsEmpty())
+	{
+		OnItemsChangedDelegate.Broadcast();
+		Items.Empty();
+		return;
+	}
+	
+	bool bMadeChanges = false;
+	TArray<TSharedPtr<FFileStateItem>> OldItems = MoveTemp(Items);
+	const FString ActiveSandbox = Model->GetActiveSandboxPath();
+	Model->EnumerateFileChanges(ActiveSandbox, 
+		[this, &bMadeChanges, &OldItems](const FString& InNonSandboxPath, FileSandboxCore::ESandboxFileChange InAction)
+		{
+			const TSharedPtr<FFileStateItem>* ReusedItem = OldItems.FindByPredicate([InNonSandboxPath, InAction](const TSharedPtr<FFileStateItem>& Item)
+			{
+				return FPaths::IsSamePath(Item->NonSandboxFile, InNonSandboxPath);
+			});
+			if (ReusedItem)
+			{
+				ReusedItem->Get()->Action = InAction; // Do this to make sure the item is up to date
+				Items.Emplace(*ReusedItem);
+			}
+			else
+			{
+				bMadeChanges = true;
+				Items.Emplace(MakeShared<FFileStateItem>(InNonSandboxPath, InAction));
+			}
+				
+			return FileSandboxCore::EBreakBehavior::Continue;
+		});
+	
+	bMadeChanges |= OldItems.Num() != Items.Num();
+	if (bMadeChanges)
+	{
+		OnItemsChangedDelegate.Broadcast();
+	}
+}
+
+void FActiveSandboxFileStateViewModel::OnSandboxLoaded()
+{
+	RefreshItems();
+}
+
+void FActiveSandboxFileStateViewModel::OnSandboxLeft()
+{
+	RefreshItems();
+}
+}
diff --git a/Engine/Plugins/Developer/Sandbox/SandboxedEditing/Source/SandboxedEditing/Private/Features/Browser/ViewModels/FileState/Models/ActiveSandboxFileStateViewModel.h b/Engine/Plugins/Developer/Sandbox/SandboxedEditing/Source/SandboxedEditing/Private/Features/Browser/ViewModels/FileState/Models/ActiveSandboxFileStateViewModel.h
new file mode 100644
--- /dev/null
+++ b/Engine/Plugins/Developer/Sandbox/SandboxedEditing/Source/SandboxedEditing/Private/Features/Browser/ViewModels/FileState/Models/ActiveSandboxFileStateViewModel.h
@@ -0,0 +1,34 @@
+﻿// Copyright Epic Games, Inc. All Rights Reserved.
+
+#pragma once
+
+#include "StaticFileStateViewModel.h"
+
+namespace UE::SandboxedEditing
+{
+class FSandboxSystemModel;
+
+/** Refreshes the file states of the active sandbox.  */
+class FActiveSandboxFileStateViewModel : public FStaticFileStateViewModel
+{
+public:
+	
+	explicit FActiveSandboxFileStateViewModel(
+		const TSharedRef<FSandboxSystemModel>& InModel,
+		TMap<FName, TSharedRef<IFileStateColumnBehavior>> InColumns
+		);
+	~FActiveSandboxFileStateViewModel();
+	
+private:
+	
+	/** Used to get the file states from the active sandbox. */
+	TSharedPtr<FSandboxSystemModel> Model;
+	
+	/** Refreshes the items from the active sandbox. */
+	void RefreshItems();
+	
+	void OnSandboxLoaded();
+	void OnSandboxLeft();
+};
+}
+
diff --git a/Engine/Plugins/Developer/Sandbox/SandboxedEditing/Source/SandboxedEditing/Private/Features/Browser/ViewModels/FileState/Models/StaticFileStateViewModel.cpp b/Engine/Plugins/Developer/Sandbox/SandboxedEditing/Source/SandboxedEditing/Private/Features/Browser/ViewModels/FileState/Models/StaticFileStateViewModel.cpp
new file mode 100644
--- /dev/null
+++ b/Engine/Plugins/Developer/Sandbox/SandboxedEditing/Source/SandboxedEditing/Private/Features/Browser/ViewModels/FileState/Models/StaticFileStateViewModel.cpp
@@ -0,0 +1,21 @@
+﻿// Copyright Epic Games, Inc. All Rights Reserved.
+
+#include "StaticFileStateViewModel.h"
+
+namespace UE::SandboxedEditing
+{
+FStaticFileStateViewModel::FStaticFileStateViewModel(
+	TMap<FName, TSharedRef<IFileStateColumnBehavior>> InColumns,
+	TArray<TSharedPtr<FFileStateItem>> InItems
+	)
+	: Columns(MoveTemp(InColumns))
+	, Items(MoveTemp(InItems))
+{}
+
+TArray<FName> FStaticFileStateViewModel::GetDisplayedColumns() const
+{
+	TArray<FName> ColumnNames;
+	Columns.GenerateKeyArray(ColumnNames);
+	return ColumnNames;
+}
+}
diff --git a/Engine/Plugins/Developer/Sandbox/SandboxedEditing/Source/SandboxedEditing/Private/Features/Browser/ViewModels/FileState/Models/StaticFileStateViewModel.h b/Engine/Plugins/Developer/Sandbox/SandboxedEditing/Source/SandboxedEditing/Private/Features/Browser/ViewModels/FileState/Models/StaticFileStateViewModel.h
new file mode 100644
--- /dev/null
+++ b/Engine/Plugins/Developer/Sandbox/SandboxedEditing/Source/SandboxedEditing/Private/Features/Browser/ViewModels/FileState/Models/StaticFileStateViewModel.h
@@ -0,0 +1,49 @@
+﻿// Copyright Epic Games, Inc. All Rights Reserved.
+
+#pragma once
+
+#include "Features/Browser/ViewModels/FileState/IFileStateViewModel.h"
+
+namespace UE::SandboxedEditing
+{
+class IFileStateColumnBehavior;
+class FSandboxListItem;
+
+/**
+ * Initialized with items at construction. 
+ * 
+ * The item array is not directly mutated by this class (TODO except for sorting columns). 
+ * Subclasses may mutate the item array.
+ * This can serve as a base implementation for more dynamic models.
+ */
+class FStaticFileStateViewModel : public IFileStateViewModel
+{
+public:
+	
+	explicit FStaticFileStateViewModel(
+		TMap<FName, TSharedRef<IFileStateColumnBehavior>> InColumns,
+		TArray<TSharedPtr<FFileStateItem>> InItems = {}
+		);
+	
+	//~ Begin IFileStateViewModel Interface
+	virtual const TArray<TSharedPtr<FFileStateItem>>& GetItems() const override { return Items; }
+	virtual TArray<FName> GetDisplayedColumns() const override;
+	virtual FSimpleMulticastDelegate& OnItemsChanged() override { return OnItemsChangedDelegate; }
+	//~ End IFileStateViewModel Interface
+
+protected:
+	
+	/** All columns that can ever be displayed. */
+	const TMap<FName, TSharedRef<IFileStateColumnBehavior>> Columns;
+	
+	/** The items to display. */
+	TArray<TSharedPtr<FFileStateItem>> Items;
+	
+	/**
+	 * Dummy to fulfill the interface requirements. Can be used by subclasses. 
+	 * TODO when sorting column is implemented, this can be invoked.
+	 */
+	FSimpleMulticastDelegate OnItemsChangedDelegate;
+};
+}
+
diff --git a/Engine/Plugins/Developer/Sandbox/SandboxedEditing/Source/SandboxedEditing/Private/Features/Browser/ViewModels/FileState/Models/UnloadedSandboxFileStateViewModel.cpp b/Engine/Plugins/Developer/Sandbox/SandboxedEditing/Source/SandboxedEditing/Private/Features/Browser/ViewModels/FileState/Models/UnloadedSandboxFileStateViewModel.cpp
new file mode 100644
--- /dev/null
+++ b/Engine/Plugins/Developer/Sandbox/SandboxedEditing/Source/SandboxedEditing/Private/Features/Browser/ViewModels/FileState/Models/UnloadedSandboxFileStateViewModel.cpp
@@ -0,0 +1,49 @@
+﻿// Copyright Epic Games, Inc. All Rights Reserved.
+
+#include "UnloadedSandboxFileStateViewModel.h"
+
+#include "Features/Browser/ViewModels/FileState/FileStateItemUtils.h"
+#include "Types/GatheredFileChanges.h"
+
+namespace UE::SandboxedEditing
+{
+FUnloadedSandboxFileStateViewModel::FUnloadedSandboxFileStateViewModel(
+	const TSharedRef<FSandboxSystemModel>& InModel,
+	const TMap<FName, TSharedRef<IFileStateColumnBehavior>>& InColumns
+	)
+	: FStaticFileStateViewModel(InColumns, {})
+	, Model(InModel)
+{
+	Model->OnLeaveSandbox().AddRaw(this, &FUnloadedSandboxFileStateViewModel::OnSandboxLeft);
+}
+
+FUnloadedSandboxFileStateViewModel::~FUnloadedSandboxFileStateViewModel()
+{
+	Model->OnLeaveSandbox().RemoveAll(this);
+}
+
+void FUnloadedSandboxFileStateViewModel::SetContent(const FString& InSandboxRoot)
+{
+	SelectedSandbox = InSandboxRoot;
+	
+	Items = MakeItemsFromFileChanges(
+		Model->GatherFileChanges(InSandboxRoot)
+		);
+	OnItemsChangedDelegate.Broadcast();
+}
+
+void FUnloadedSandboxFileStateViewModel::ClearContent()
+{
+	SelectedSandbox.Reset();
+	Items.Empty();
+	OnItemsChangedDelegate.Broadcast();
+}
+
+void FUnloadedSandboxFileStateViewModel::OnSandboxLeft()
+{
+	if (SelectedSandbox)
+	{
+		SetContent(*SelectedSandbox);
+	}
+}
+}
diff --git a/Engine/Plugins/Developer/Sandbox/SandboxedEditing/Source/SandboxedEditing/Private/Features/Browser/ViewModels/FileState/Models/UnloadedSandboxFileStateViewModel.h b/Engine/Plugins/Developer/Sandbox/SandboxedEditing/Source/SandboxedEditing/Private/Features/Browser/ViewModels/FileState/Models/UnloadedSandboxFileStateViewModel.h
new file mode 100644
--- /dev/null
+++ b/Engine/Plugins/Developer/Sandbox/SandboxedEditing/Source/SandboxedEditing/Private/Features/Browser/ViewModels/FileState/Models/UnloadedSandboxFileStateViewModel.h
@@ -0,0 +1,42 @@
+﻿// Copyright Epic Games, Inc. All Rights Reserved.
+
+#pragma once
+
+#include "StaticFileStateViewModel.h"
+#include "Framework/Models/SandboxSystemModel.h"
+
+namespace UE::SandboxedEditing
+{
+/**
+ * View model for displaying the file changes of an unloaded sandbox, i.e. when the engine is not sandboxed.
+ * Used to display info about a sandbox in the browser.
+ */
+class FUnloadedSandboxFileStateViewModel : public FStaticFileStateViewModel
+{
+public:
+	
+	explicit FUnloadedSandboxFileStateViewModel(
+		const TSharedRef<FSandboxSystemModel>& InModel,
+		const TMap<FName, TSharedRef<IFileStateColumnBehavior>>& InColumns
+		);
+	~FUnloadedSandboxFileStateViewModel();
+	
+	/** Sets the content to be that of the specified sandbox. */
+	void SetContent(const FString& InSandboxRoot);
+	
+	/** Clears the content. */
+	void ClearContent();
+	
+private:
+	
+	/** The model through which we interact with the sandbox system. */
+	const TSharedRef<FSandboxSystemModel> Model;
+	
+	/** The sandbox being viewed. */
+	TOptional<FString> SelectedSandbox;
+	
+	/** Refreshes the content, in case the active sandbox changed files. */
+	void OnSandboxLeft();
+};
+}
+
diff --git a/Engine/Plugins/Developer/Sandbox/SandboxedEditing/Source/SandboxedEditing/Private/Features/Browser/ViewModels/FileState/StaticFileStateViewModel.cpp b/Engine/Plugins/Developer/Sandbox/SandboxedEditing/Source/SandboxedEditing/Private/Features/Browser/ViewModels/FileState/StaticFileStateViewModel.cpp
deleted file mode 100644
diff --git a/Engine/Plugins/Developer/Sandbox/SandboxedEditing/Source/SandboxedEditing/Private/Features/Browser/ViewModels/FileState/StaticFileStateViewModel.h b/Engine/Plugins/Developer/Sandbox/SandboxedEditing/Source/SandboxedEditing/Private/Features/Browser/ViewModels/FileState/StaticFileStateViewModel.h
deleted file mode 100644
diff --git a/Engine/Plugins/Developer/Sandbox/SandboxedEditing/Source/SandboxedEditing/Private/Features/Browser/ViewModels/FileState/UnloadedSandboxFileStateViewModel.cpp b/Engine/Plugins/Developer/Sandbox/SandboxedEditing/Source/SandboxedEditing/Private/Features/Browser/ViewModels/FileState/UnloadedSandboxFileStateViewModel.cpp
deleted file mode 100644
diff --git a/Engine/Plugins/Developer/Sandbox/SandboxedEditing/Source/SandboxedEditing/Private/Features/Browser/ViewModels/FileState/UnloadedSandboxFileStateViewModel.h b/Engine/Plugins/Developer/Sandbox/SandboxedEditing/Source/SandboxedEditing/Private/Features/Browser/ViewModels/FileState/UnloadedSandboxFileStateViewModel.h
deleted file mode 100644
diff --git a/Engine/Plugins/Developer/Sandbox/SandboxedEditing/Source/SandboxedEditing/Private/Features/Browser/Widgets/SSandboxedEditingRoot.cpp b/Engine/Plugins/Developer/Sandbox/SandboxedEditing/Source/SandboxedEditing/Private/Features/Browser/Widgets/SSandboxedEditingRoot.cpp
--- a/Engine/Plugins/Developer/Sandbox/SandboxedEditing/Source/SandboxedEditing/Private/Features/Browser/Widgets/SSandboxedEditingRoot.cpp
+++ b/Engine/Plugins/Developer/Sandbox/SandboxedEditing/Source/SandboxedEditing/Private/Features/Browser/Widgets/SSandboxedEditingRoot.cpp
@@ -19,7 +19,10 @@ void SSandboxedEditingRoot::Construct(const FArguments& InArgs, const FBrowserVi
 	CommandBindings = InArgs._CommandList;
 	ActiveSandboxTrackerViewModel = InViewModels.ActiveSandboxTrackerViewModel;
 	ExternalSandboxViewModel = InViewModels.ExternalSandboxViewModel;
-	ActiveSandboxViewModels.Emplace(InViewModels.ActiveSandboxDetailsViewModel, InViewModels.MetaDataViewModel);
+	ActiveSandboxViewModels.Emplace(
+		InViewModels.ActiveSandboxDetailsViewModel, InViewModels.MetaDataViewModel, InViewModels.ActiveSandboxFileStateViewModel
+		);
+	ActiveFileStateColumnsAttr = InArgs._ActiveFileStateColumnsFactories;
 	
 	ActiveSandboxTrackerViewModel->OnLoadSandbox().AddSP(this, &SSandboxedEditingRoot::UpdateWidgetContent);
 	ActiveSandboxTrackerViewModel->OnLeaveSandbox().AddSP(this, &SSandboxedEditingRoot::ShowBrowser);
@@ -33,7 +36,7 @@ void SSandboxedEditingRoot::Construct(const FArguments& InArgs, const FBrowserVi
 	Browser = SNew(SSandboxedEditingBrowser, ViewModels)
 		.CommandList(CommandBindings)
 		.SandboxColumnFactories(InArgs._SandboxColumnFactories)
-		.FileActionsColumnFactories(InArgs._FileActionsColumnFactories);
+		.FileActionsColumnFactories(InArgs._BrowserFileActionsColumnFactories);
 	
 	UpdateWidgetContent();
 }
@@ -66,6 +69,7 @@ void SSandboxedEditingRoot::ShowActiveSandbox()
 	[
 		SNew(SActiveSandbox, *ActiveSandboxViewModels)
 		.CommandList(CommandBindings)
+		.ColumnFactories(ActiveFileStateColumnsAttr.Get())
 	];
 }
 
diff --git a/Engine/Plugins/Developer/Sandbox/SandboxedEditing/Source/SandboxedEditing/Private/Features/Browser/Widgets/SSandboxedEditingRoot.h b/Engine/Plugins/Developer/Sandbox/SandboxedEditing/Source/SandboxedEditing/Private/Features/Browser/Widgets/SSandboxedEditingRoot.h
--- a/Engine/Plugins/Developer/Sandbox/SandboxedEditing/Source/SandboxedEditing/Private/Features/Browser/Widgets/SSandboxedEditingRoot.h
+++ b/Engine/Plugins/Developer/Sandbox/SandboxedEditing/Source/SandboxedEditing/Private/Features/Browser/Widgets/SSandboxedEditingRoot.h
@@ -13,11 +13,11 @@ namespace UE::FileSandboxUI { class IExternalSandboxActiveViewModel; }
 
 namespace UE::SandboxedEditing
 {
-class FSandboxMetaDataViewModel;
 class FActiveSandboxTrackerViewModel;
 class FBrowserCommandBindings;
 class FSandboxControlsViewModel;
 class FSandboxListViewModel;
+class FSandboxMetaDataViewModel;
 class SSandboxedEditingBrowser;
 struct FBrowserViewModels;
 
@@ -29,10 +29,14 @@ public:
 	SLATE_BEGIN_ARGS(SSandboxedEditingRoot){}
 		/** The command list to bind selection-based commands to. */
 		SLATE_ARGUMENT(TSharedPtr<FUICommandList>, CommandList)
+		
 		/** Used to build the columns. */
 		SLATE_ARGUMENT(FSandboxColumnFactoryMap, SandboxColumnFactories)
-		/** Used to build the columns for the file actions list. */
-		SLATE_ARGUMENT(FFileStateColumnFactoryMap, FileActionsColumnFactories)
+		/** Used to build the columns for the file actions list displayed in the browser. */
+		SLATE_ARGUMENT(FFileStateColumnFactoryMap, BrowserFileActionsColumnFactories)
+		
+		/** Used to build the columns for the file actions list for the active sandbox. */
+		SLATE_ARGUMENT(FFileStateColumnFactoryMap, ActiveFileStateColumnsFactories)
 	SLATE_END_ARGS()
 
 	void Construct(const FArguments& InArgs, const FBrowserViewModels& InViewModels);
@@ -50,8 +54,11 @@ private:
 	TSharedPtr<FActiveSandboxTrackerViewModel> ActiveSandboxTrackerViewModel;
 	/** Knows when an external sandbox is active. */
 	TSharedPtr<FileSandboxUI::IExternalSandboxActiveViewModel> ExternalSandboxViewModel;
+	
 	/** The view models required to construct the active sandbox widget. Set once Construct has been called. */
 	TOptional<SActiveSandbox::FViewModels> ActiveSandboxViewModels;
+	/** Used to build the columns for the file actions list for the active sandbox. */
+	TAttribute<FFileStateColumnFactoryMap> ActiveFileStateColumnsAttr;
 	
 	/** Command bindings for the browser UI. */
 	TSharedPtr<FUICommandList> CommandBindings;
diff --git a/Engine/Plugins/Developer/Sandbox/SandboxedEditing/Source/SandboxedEditing/Private/Features/Browser/Widgets/Sandboxed/SActiveSandbox.cpp b/Engine/Plugins/Developer/Sandbox/SandboxedEditing/Source/SandboxedEditing/Private/Features/Browser/Widgets/Sandboxed/SActiveSandbox.cpp
--- a/Engine/Plugins/Developer/Sandbox/SandboxedEditing/Source/SandboxedEditing/Private/Features/Browser/Widgets/Sandboxed/SActiveSandbox.cpp
+++ b/Engine/Plugins/Developer/Sandbox/SandboxedEditing/Source/SandboxedEditing/Private/Features/Browser/Widgets/Sandboxed/SActiveSandbox.cpp
@@ -3,7 +3,9 @@
 #include "SActiveSandbox.h"
 
 #include "Features/Browser/ViewModels/Active/ActiveSandboxDetailsViewModel.h"
+#include "Features/Browser/ViewModels/FileState/Models/ActiveSandboxFileStateViewModel.h"
 #include "Features/Browser/Widgets/Shared/SSandboxDescription.h"
+#include "FileState/SFileStateListView.h"
 #include "SActiveSandboxToolbar.h"
 #include "Styling/ToolBarStyle.h"
 #include "Widgets/Layout/SBorder.h"
@@ -33,7 +35,7 @@ void SActiveSandbox::Construct(const FArguments& InArgs, const FViewModels& InVi
 			.Padding(1)
 			[
 				SNew(SActiveSandboxToolbar, InViewModels.ActiveSandboxViewModel)
-					.CommandList(InArgs._CommandList)
+				.CommandList(InArgs._CommandList)
 			]
 
 			+SVerticalBox::Slot()
@@ -50,6 +52,13 @@ void SActiveSandbox::Construct(const FArguments& InArgs, const FViewModels& InVi
 			[
 				MakeDescription(InViewModels)
 			]
+
+			+SVerticalBox::Slot()
+			.FillHeight(1.f)
+			.Padding(1)
+			[
+				MakeFileChangesContent(InArgs, InViewModels)
+			]
 		]
 	];
 }
@@ -73,6 +82,26 @@ TSharedRef<SWidget> SActiveSandbox::MakeDescription(const FViewModels& InViewMod
 			]
 		];
 }
+
+TSharedRef<SWidget> SActiveSandbox::MakeFileChangesContent(const FArguments& InArgs, const FViewModels& InViewModels)
+{
+	return SNew(SExpandableArea)
+		.BodyBorderImage(FAppStyle::Get().GetBrush("ToolPanel.GroupBorder"))
+		.HeaderContent()
+		[
+			SNew(STextBlock)
+			.Text(LOCTEXT("FileChanges.Label", "Changes"))
+			.ToolTipText(LOCTEXT("FileChanges.Tooltip", "The file changes made in this sandbox"))
+		]
+		.BodyContent()
+		[
+			SNew(SBorder)
+			[
+				SNew(SFileStateListView, InViewModels.ActiveSandboxFileStateModel)
+				.ColumnFactories(InArgs._ColumnFactories)
+			]
+		];
+}
 }
 
 #undef LOCTEXT_NAMESPACE
\ No newline at end of file
diff --git a/Engine/Plugins/Developer/Sandbox/SandboxedEditing/Source/SandboxedEditing/Private/Features/Browser/Widgets/Sandboxed/SActiveSandbox.h b/Engine/Plugins/Developer/Sandbox/SandboxedEditing/Source/SandboxedEditing/Private/Features/Browser/Widgets/Sandboxed/SActiveSandbox.h
--- a/Engine/Plugins/Developer/Sandbox/SandboxedEditing/Source/SandboxedEditing/Private/Features/Browser/Widgets/Sandboxed/SActiveSandbox.h
+++ b/Engine/Plugins/Developer/Sandbox/SandboxedEditing/Source/SandboxedEditing/Private/Features/Browser/Widgets/Sandboxed/SActiveSandbox.h
@@ -2,6 +2,7 @@
 
 #pragma once
 
+#include "Features/Browser/ViewModels/FileState/FileStateColumns.h"
 #include "Widgets/DeclarativeSyntaxSupport.h"
 #include "Widgets/SCompoundWidget.h"
 
@@ -9,6 +10,7 @@ class FUICommandList;
 
 namespace UE::SandboxedEditing
 {
+class FActiveSandboxFileStateViewModel;
 class FActiveSandboxDetailsViewModel;
 class FSandboxMetaDataViewModel;
 
@@ -21,19 +23,25 @@ public:
 	{
 		TSharedRef<FActiveSandboxDetailsViewModel> ActiveSandboxViewModel;
 		TSharedRef<FSandboxMetaDataViewModel> MetaDataViewModel;
+		TSharedRef<FActiveSandboxFileStateViewModel> ActiveSandboxFileStateModel;
 
 		explicit FViewModels(
 			const TSharedRef<FActiveSandboxDetailsViewModel>& InActiveSandboxViewModel,
-			const TSharedRef<FSandboxMetaDataViewModel>& InMetaDataViewModel
+			const TSharedRef<FSandboxMetaDataViewModel>& InMetaDataViewModel,
+			const TSharedRef<FActiveSandboxFileStateViewModel>& InActiveSandboxFileStateModel
 			)
 			: ActiveSandboxViewModel(InActiveSandboxViewModel)
 			, MetaDataViewModel(InMetaDataViewModel)
+			, ActiveSandboxFileStateModel(InActiveSandboxFileStateModel)
 		{}
 	};
 	
 	SLATE_BEGIN_ARGS(SActiveSandbox) {}
 		/** The command list to bind commands to. */
 		SLATE_ARGUMENT(TSharedPtr<FUICommandList>, CommandList)
+		
+		/** Used to build the columns for the file actions list. */
+		SLATE_ARGUMENT(FFileStateColumnFactoryMap, ColumnFactories)
 	SLATE_END_ARGS()
 
 	void Construct(const FArguments& InArgs, const FViewModels& InViewModels);
@@ -45,6 +53,9 @@ private:
 	
 	/** Make the description widget. */
 	TSharedRef<SWidget> MakeDescription(const FViewModels& InViewModels);
+	
+	/** Makes the widget that displays the files changed by the current sandbox. */
+	TSharedRef<SWidget> MakeFileChangesContent(const FArguments& InArgs, const FViewModels& InViewModels);
 };
 }
 
diff --git a/Engine/Plugins/Developer/Sandbox/SandboxedEditing/Source/SandboxedEditing/Private/Features/Browser/Widgets/Sandboxed/SPersistOperationWidget.cpp b/Engine/Plugins/Developer/Sandbox/SandboxedEditing/Source/SandboxedEditing/Private/Features/Browser/Widgets/Sandboxed/SPersistOperationWidget.cpp
--- a/Engine/Plugins/Developer/Sandbox/SandboxedEditing/Source/SandboxedEditing/Private/Features/Browser/Widgets/Sandboxed/SPersistOperationWidget.cpp
+++ b/Engine/Plugins/Developer/Sandbox/SandboxedEditing/Source/SandboxedEditing/Private/Features/Browser/Widgets/Sandboxed/SPersistOperationWidget.cpp
@@ -2,15 +2,12 @@
 
 #include "SPersistOperationWidget.h"
 
+#include "Features/Browser/ViewModels/FileState/FileStateColumnRegistry.h"
 #include "Features/Browser/ViewModels/FileState/FileStateItemUtils.h"
-#include "Features/Browser/ViewModels/FileState/StaticFileStateViewModel.h"
+#include "Features/Browser/ViewModels/FileState/Models/StaticFileStateViewModel.h"
 #include "Features/Browser/ViewModels/Persist/PersistOperationViewModel.h"
-#include "FileState/Columns/ActionFileStateColumn.h"
-#include "FileState/Columns/PathFileStateColumn.h"
-#include "FileState/Columns/PersistFileStateColumn.h"
 #include "FileState/SFileStateListView.h"
 #include "SPositiveActionButton.h"
-#include "Features/Browser/ViewModels/FileState/FileStateColumnRegistry.h"
 #include "Styling/AppStyle.h"
 #include "Types/GatheredFileChanges.h"
 #include "Widgets/Input/SButton.h"
diff --git a/Engine/Plugins/Developer/Sandbox/SandboxedEditing/Source/SandboxedEditing/Private/Features/Browser/Widgets/Unsandboxed/List/Columns/NameSandboxColumn.cpp b/Engine/Plugins/Developer/Sandbox/SandboxedEditing/Source/SandboxedEditing/Private/Features/Browser/Widgets/Unsandboxed/List/Columns/NameSandboxColumn.cpp
--- a/Engine/Plugins/Developer/Sandbox/SandboxedEditing/Source/SandboxedEditing/Private/Features/Browser/Widgets/Unsandboxed/List/Columns/NameSandboxColumn.cpp
+++ b/Engine/Plugins/Developer/Sandbox/SandboxedEditing/Source/SandboxedEditing/Private/Features/Browser/Widgets/Unsandboxed/List/Columns/NameSandboxColumn.cpp
@@ -36,9 +36,10 @@ TSharedRef<SWidget> FNameSandboxColumn::MakeColumnWidget(const FMakeSandboxColum
 		.Padding(FSandboxedEditingStyle::Get().GetMargin("SandboxedEditing.Browser.NameColumn.Padding"))
 		[
 			SNew(SInlineEditableTextBlock)
-			.Text_Lambda([Name = InArgs.RowData->SandboxInfo.Name] { return FText::AsCultureInvariant(Name); })
-			.ToolTipText_Lambda([Description = InArgs.RowData->SandboxInfo.Description]
+			.Text_Lambda([RowData = InArgs.RowData] { return FText::AsCultureInvariant(RowData->SandboxInfo.Name); })
+			.ToolTipText_Lambda([RowData = InArgs.RowData]
 			{
+				const FString& Description = RowData->SandboxInfo.Description;
 				return !Description.IsEmpty()
 					? FText::AsCultureInvariant(Description) 
 					: LOCTEXT("NoDescription", "This sandbox has no user-provided description");
diff --git a/Engine/Plugins/Developer/Sandbox/SandboxedEditing/Source/SandboxedEditing/Private/Features/Browser/Widgets/Unsandboxed/SBrowserSandboxDetails.cpp b/Engine/Plugins/Developer/Sandbox/SandboxedEditing/Source/SandboxedEditing/Private/Features/Browser/Widgets/Unsandboxed/SBrowserSandboxDetails.cpp
--- a/Engine/Plugins/Developer/Sandbox/SandboxedEditing/Source/SandboxedEditing/Private/Features/Browser/Widgets/Unsandboxed/SBrowserSandboxDetails.cpp
+++ b/Engine/Plugins/Developer/Sandbox/SandboxedEditing/Source/SandboxedEditing/Private/Features/Browser/Widgets/Unsandboxed/SBrowserSandboxDetails.cpp
@@ -4,7 +4,7 @@
 
 #include "Components/VerticalBox.h"
 #include "Features/Browser/ViewModels/FileState/FileStateColumnRegistry.h"
-#include "Features/Browser/ViewModels/FileState/UnloadedSandboxFileStateViewModel.h"
+#include "Features/Browser/ViewModels/FileState/Models/UnloadedSandboxFileStateViewModel.h"
 #include "Features/Browser/Widgets/Sandboxed/FileState/SFileStateListView.h"
 #include "Features/Browser/Widgets/Shared/SSandboxDescription.h"
 #include "Internationalization/Text.h"
diff --git a/Engine/Plugins/Developer/Sandbox/SandboxedEditing/Source/SandboxedEditing/Private/Features/Browser/Widgets/Unsandboxed/SSandboxedEditingBrowser.cpp b/Engine/Plugins/Developer/Sandbox/SandboxedEditing/Source/SandboxedEditing/Private/Features/Browser/Widgets/Unsandboxed/SSandboxedEditingBrowser.cpp
--- a/Engine/Plugins/Developer/Sandbox/SandboxedEditing/Source/SandboxedEditing/Private/Features/Browser/Widgets/Unsandboxed/SSandboxedEditingBrowser.cpp
+++ b/Engine/Plugins/Developer/Sandbox/SandboxedEditing/Source/SandboxedEditing/Private/Features/Browser/Widgets/Unsandboxed/SSandboxedEditingBrowser.cpp
@@ -3,6 +3,8 @@
 #include "SSandboxedEditingBrowser.h"
 
 #include "Features/Browser/Commands/BrowserCommandBindings.h"
+#include "Features/Browser/ViewModels/FileState/Models/UnloadedSandboxFileStateViewModel.h"
+#include "Features/Browser/ViewModels/List/FilterSandboxViewModel.h"
 #include "Framework/Commands/GenericCommands.h"
 #include "GenericBrowser/SSessionBrowserFrame.h"
 #include "List/SSandboxListView.h"
@@ -10,8 +12,6 @@
 #include "SBrowserControls.h"
 #include "SBrowserFilters.h"
 #include "SBrowserSandboxDetails.h"
-#include "Features/Browser/ViewModels/FileState/UnloadedSandboxFileStateViewModel.h"
-#include "Features/Browser/ViewModels/List/FilterSandboxViewModel.h"
 #include "Utils/BreakBehavior.h"
 #include "Widgets/Text/STextBlock.h"
 
diff --git a/Engine/Plugins/Developer/Sandbox/SandboxedEditing/Source/SandboxedEditing/Private/Framework/Models/SandboxSystemModel.cpp b/Engine/Plugins/Developer/Sandbox/SandboxedEditing/Source/SandboxedEditing/Private/Framework/Models/SandboxSystemModel.cpp
--- a/Engine/Plugins/Developer/Sandbox/SandboxedEditing/Source/SandboxedEditing/Private/Framework/Models/SandboxSystemModel.cpp
+++ b/Engine/Plugins/Developer/Sandbox/SandboxedEditing/Source/SandboxedEditing/Private/Framework/Models/SandboxSystemModel.cpp
@@ -33,12 +33,12 @@ FSandboxSystemModel::FSandboxSystemModel()
 	IFileSandboxCoreModule& Module = IFileSandboxCoreModule::Get();
 	
 	ISandboxManager& Manager = Module.GetSandboxManager();
-	Manager.OnPostSandboxStartup().AddRaw(this, &FSandboxSystemModel::BroadcastOnLoadSandbox);
-	Manager.OnPostSandboxShutdown().AddRaw(this, &FSandboxSystemModel::BroadcastOnLeaveSandbox);
+	Manager.OnPostSandboxStartup().AddRaw(this, &FSandboxSystemModel::HandleLoadSandbox);
+	Manager.OnPostSandboxShutdown().AddRaw(this, &FSandboxSystemModel::HandleLeaveSandbox);
 	
 	ISandboxRepository& Repository = Module.GetDefaultSandboxRepository();
-	Repository.OnSandboxesChanged().AddRaw(this, &FSandboxSystemModel::OnSandboxesChanged);
-	Repository.OnSandboxMetaDataChanged().AddRaw(this, &FSandboxSystemModel::OnMetaDataChanged);
+	Repository.OnSandboxesChanged().AddRaw(this, &FSandboxSystemModel::HandleSandboxesChanged);
+	Repository.OnSandboxMetaDataChanged().AddRaw(this, &FSandboxSystemModel::HandleMetaDataChanged);
 }
 
 FSandboxSystemModel::~FSandboxSystemModel()
@@ -52,6 +52,10 @@ FSandboxSystemModel::~FSandboxSystemModel()
 		ISandboxManager& Manager = Module.GetSandboxManager();
 		Manager.OnPostSandboxStartup().RemoveAll(this);
 		Manager.OnPostSandboxShutdown().RemoveAll(this);
+		if (ISandboxInstance* Instance = Manager.GetActiveSandboxInstance())
+		{
+			Instance->OnSandboxedFilesChanged().RemoveAll(this);
+		}
 		
 		ISandboxRepository& Repository = Module.GetDefaultSandboxRepository();
 		Repository.OnSandboxesChanged().RemoveAll(this);
@@ -275,6 +279,17 @@ FileSandboxCore::FGatheredFileChanges FSandboxSystemModel::GatherFileChanges(con
 	return Result;
 }
 
+void FSandboxSystemModel::EnumerateFileChanges(
+	const FString& InSandbox,
+	TFunctionRef<FileSandboxCore::EBreakBehavior(const FString& InPath, FileSandboxCore::ESandboxFileChange InChange)> InCallback
+	)
+{
+	using namespace FileSandboxCore;
+	IFileSandboxCoreModule& Module = IFileSandboxCoreModule::Get();
+	ISandboxManager& SandboxManager = Module.GetSandboxManager();
+	SandboxManager.EnumerateFileChanges(InSandbox, InCallback);
+}
+
 void FSandboxSystemModel::DiscardAllChanges() const
 {
 	using namespace FileSandboxCore;
@@ -371,6 +386,22 @@ void FSandboxSystemModel::SetDescription(const FString& InSandboxRoot, const FSt
 		}
 	});
 }
+
+void FSandboxSystemModel::HandleLoadSandbox(FileSandboxCore::ISandboxInstance& InInstance)
+{
+	OnLoadSandboxDelegate.Broadcast();
+	InInstance.OnSandboxedFilesChanged().AddRaw(this, &FSandboxSystemModel::HandleActiveSandboxFilesChanged);
+}
+
+void FSandboxSystemModel::HandleLeaveSandbox()
+{
+	OnLeaveSandboxDelegate.Broadcast();
+}
+
+void FSandboxSystemModel::HandleActiveSandboxFilesChanged()
+{
+	OnSandboxFilesChangedDelegate.Broadcast();
+}
 }
 
 #undef LOCTEXT_NAMESPACE
\ No newline at end of file
diff --git a/Engine/Plugins/Developer/Sandbox/SandboxedEditing/Source/SandboxedEditing/Private/Framework/Models/SandboxSystemModel.h b/Engine/Plugins/Developer/Sandbox/SandboxedEditing/Source/SandboxedEditing/Private/Framework/Models/SandboxSystemModel.h
--- a/Engine/Plugins/Developer/Sandbox/SandboxedEditing/Source/SandboxedEditing/Private/Framework/Models/SandboxSystemModel.h
+++ b/Engine/Plugins/Developer/Sandbox/SandboxedEditing/Source/SandboxedEditing/Private/Framework/Models/SandboxSystemModel.h
@@ -9,8 +9,10 @@
 namespace UE::FileSandboxCore
 {
 class IPersistFeedback;
-struct FGatheredFileChanges;
 class ISandboxInstance;
+enum class EBreakBehavior : uint8;
+enum class ESandboxFileChange : uint8;
+struct FGatheredFileChanges;
 }
 
 template<typename OptionalType> struct TOptional;
@@ -77,6 +79,12 @@ public:
 	FileSandboxCore::FGatheredFileChanges GatherActiveFileChanges() const;
 	/** @return The file changes of the specified sandbox. */
 	FileSandboxCore::FGatheredFileChanges GatherFileChanges(const FString& InSandbox) const;
+	
+	/** Invokes InCallback for each file change. */
+	void EnumerateFileChanges(
+		const FString& InSandbox, 
+		TFunctionRef<FileSandboxCore::EBreakBehavior(const FString& InPath, FileSandboxCore::ESandboxFileChange InChange)> InCallback
+		);
 
 	/** Discards all changes */
 	void DiscardAllChanges() const;
@@ -107,6 +115,9 @@ public:
 	FSimpleMulticastDelegate& OnLoadSandbox() { return OnLoadSandboxDelegate; }
 	/** Invoked when the current sandbox is left. */
 	FSimpleMulticastDelegate& OnLeaveSandbox() { return OnLeaveSandboxDelegate; }
+	
+	/** Invoked when the files in the active sandbox have changed. */
+	FSimpleMulticastDelegate& OnSandboxFilesChanged() { return OnSandboxFilesChangedDelegate; }
 		
 private:
 	
@@ -118,11 +129,16 @@ private:
 	/** Invoked when the current sandbox is left. */
 	FSimpleMulticastDelegate OnLeaveSandboxDelegate;
 	
-	void OnSandboxesChanged(const FileSandboxCore::FRepositoryChangedEvent&) const { OnKnownSandboxesChangedDelegate.Broadcast(); }
-	void OnMetaDataChanged(const FString&) const { OnKnownSandboxesChangedDelegate.Broadcast(); }
+	/** Invoked when the files in the active sandbox have changed. */
+	FSimpleMulticastDelegate OnSandboxFilesChangedDelegate;
 	
-	void BroadcastOnLoadSandbox(FileSandboxCore::ISandboxInstance&) { OnLoadSandboxDelegate.Broadcast(); }
-	void BroadcastOnLeaveSandbox() { OnLeaveSandboxDelegate.Broadcast(); }
+	void HandleSandboxesChanged(const FileSandboxCore::FRepositoryChangedEvent&) const { OnKnownSandboxesChangedDelegate.Broadcast(); }
+	void HandleMetaDataChanged(const FString&) const { OnKnownSandboxesChangedDelegate.Broadcast(); }
+
+	void HandleLoadSandbox(FileSandboxCore::ISandboxInstance& InInstance);
+	void HandleLeaveSandbox();
+	/** Invoked when the set of files touched by the active sandbox changes. */
+	void HandleActiveSandboxFilesChanged();
 };
 }
 
