Sandboxed Editing: Moves shared files into Shared folder.

#Synced-CL 49159615

diff --git a/Engine/Plugins/Developer/Sandbox/SandboxedEditing/Source/SandboxedEditing/Private/Features/Browser/ViewModels/FileState/FileStateColumnRegistry.cpp b/Engine/Plugins/Developer/Sandbox/SandboxedEditing/Source/SandboxedEditing/Private/Features/Browser/ViewModels/FileState/FileStateColumnRegistry.cpp
--- a/Engine/Plugins/Developer/Sandbox/SandboxedEditing/Source/SandboxedEditing/Private/Features/Browser/ViewModels/FileState/FileStateColumnRegistry.cpp
+++ b/Engine/Plugins/Developer/Sandbox/SandboxedEditing/Source/SandboxedEditing/Private/Features/Browser/ViewModels/FileState/FileStateColumnRegistry.cpp
@@ -3,9 +3,9 @@
 #include "FileStateColumnRegistry.h"
 
 #include "FileStateColumns.h"
-#include "Features/Browser/Widgets/Sandboxed/FileState/Columns/ActionFileStateColumn.h"
-#include "Features/Browser/Widgets/Sandboxed/FileState/Columns/PathFileStateColumn.h"
-#include "Features/Browser/Widgets/Sandboxed/FileState/Columns/PersistFileStateColumn.h"
+#include "Features/Browser/Widgets/Shared/FileState/Columns/ActionFileStateColumn.h"
+#include "Features/Browser/Widgets/Shared/FileState/Columns/PathFileStateColumn.h"
+#include "Features/Browser/Widgets/Shared/FileState/Columns/PersistFileStateColumn.h"
 
 namespace UE::SandboxedEditing
 {
diff --git a/Engine/Plugins/Developer/Sandbox/SandboxedEditing/Source/SandboxedEditing/Private/Features/Browser/Widgets/Sandboxed/FileState/Columns/ActionFileStateColumn.cpp b/Engine/Plugins/Developer/Sandbox/SandboxedEditing/Source/SandboxedEditing/Private/Features/Browser/Widgets/Sandboxed/FileState/Columns/ActionFileStateColumn.cpp
deleted file mode 100644
diff --git a/Engine/Plugins/Developer/Sandbox/SandboxedEditing/Source/SandboxedEditing/Private/Features/Browser/Widgets/Sandboxed/FileState/Columns/ActionFileStateColumn.h b/Engine/Plugins/Developer/Sandbox/SandboxedEditing/Source/SandboxedEditing/Private/Features/Browser/Widgets/Sandboxed/FileState/Columns/ActionFileStateColumn.h
deleted file mode 100644
diff --git a/Engine/Plugins/Developer/Sandbox/SandboxedEditing/Source/SandboxedEditing/Private/Features/Browser/Widgets/Sandboxed/FileState/Columns/PathFileStateColumn.cpp b/Engine/Plugins/Developer/Sandbox/SandboxedEditing/Source/SandboxedEditing/Private/Features/Browser/Widgets/Sandboxed/FileState/Columns/PathFileStateColumn.cpp
deleted file mode 100644
diff --git a/Engine/Plugins/Developer/Sandbox/SandboxedEditing/Source/SandboxedEditing/Private/Features/Browser/Widgets/Sandboxed/FileState/Columns/PathFileStateColumn.h b/Engine/Plugins/Developer/Sandbox/SandboxedEditing/Source/SandboxedEditing/Private/Features/Browser/Widgets/Sandboxed/FileState/Columns/PathFileStateColumn.h
deleted file mode 100644
diff --git a/Engine/Plugins/Developer/Sandbox/SandboxedEditing/Source/SandboxedEditing/Private/Features/Browser/Widgets/Sandboxed/FileState/Columns/PersistFileStateColumn.cpp b/Engine/Plugins/Developer/Sandbox/SandboxedEditing/Source/SandboxedEditing/Private/Features/Browser/Widgets/Sandboxed/FileState/Columns/PersistFileStateColumn.cpp
deleted file mode 100644
diff --git a/Engine/Plugins/Developer/Sandbox/SandboxedEditing/Source/SandboxedEditing/Private/Features/Browser/Widgets/Sandboxed/FileState/Columns/PersistFileStateColumn.h b/Engine/Plugins/Developer/Sandbox/SandboxedEditing/Source/SandboxedEditing/Private/Features/Browser/Widgets/Sandboxed/FileState/Columns/PersistFileStateColumn.h
deleted file mode 100644
diff --git a/Engine/Plugins/Developer/Sandbox/SandboxedEditing/Source/SandboxedEditing/Private/Features/Browser/Widgets/Sandboxed/FileState/SFileStateListRow.cpp b/Engine/Plugins/Developer/Sandbox/SandboxedEditing/Source/SandboxedEditing/Private/Features/Browser/Widgets/Sandboxed/FileState/SFileStateListRow.cpp
deleted file mode 100644
diff --git a/Engine/Plugins/Developer/Sandbox/SandboxedEditing/Source/SandboxedEditing/Private/Features/Browser/Widgets/Sandboxed/FileState/SFileStateListRow.h b/Engine/Plugins/Developer/Sandbox/SandboxedEditing/Source/SandboxedEditing/Private/Features/Browser/Widgets/Sandboxed/FileState/SFileStateListRow.h
deleted file mode 100644
diff --git a/Engine/Plugins/Developer/Sandbox/SandboxedEditing/Source/SandboxedEditing/Private/Features/Browser/Widgets/Sandboxed/FileState/SFileStateListView.cpp b/Engine/Plugins/Developer/Sandbox/SandboxedEditing/Source/SandboxedEditing/Private/Features/Browser/Widgets/Sandboxed/FileState/SFileStateListView.cpp
deleted file mode 100644
diff --git a/Engine/Plugins/Developer/Sandbox/SandboxedEditing/Source/SandboxedEditing/Private/Features/Browser/Widgets/Sandboxed/FileState/SFileStateListView.h b/Engine/Plugins/Developer/Sandbox/SandboxedEditing/Source/SandboxedEditing/Private/Features/Browser/Widgets/Sandboxed/FileState/SFileStateListView.h
deleted file mode 100644
diff --git a/Engine/Plugins/Developer/Sandbox/SandboxedEditing/Source/SandboxedEditing/Private/Features/Browser/Widgets/Sandboxed/SActiveSandbox.cpp b/Engine/Plugins/Developer/Sandbox/SandboxedEditing/Source/SandboxedEditing/Private/Features/Browser/Widgets/Sandboxed/SActiveSandbox.cpp
--- a/Engine/Plugins/Developer/Sandbox/SandboxedEditing/Source/SandboxedEditing/Private/Features/Browser/Widgets/Sandboxed/SActiveSandbox.cpp
+++ b/Engine/Plugins/Developer/Sandbox/SandboxedEditing/Source/SandboxedEditing/Private/Features/Browser/Widgets/Sandboxed/SActiveSandbox.cpp
@@ -4,8 +4,8 @@
 
 #include "Features/Browser/ViewModels/Active/ActiveSandboxDetailsViewModel.h"
 #include "Features/Browser/ViewModels/FileState/Models/ActiveSandboxFileStateViewModel.h"
+#include "Features/Browser/Widgets/Shared/FileState/SFileStateListView.h"
 #include "Features/Browser/Widgets/Shared/SSandboxDescription.h"
-#include "FileState/SFileStateListView.h"
 #include "SActiveSandboxToolbar.h"
 #include "Styling/ToolBarStyle.h"
 #include "Widgets/Layout/SBorder.h"
diff --git a/Engine/Plugins/Developer/Sandbox/SandboxedEditing/Source/SandboxedEditing/Private/Features/Browser/Widgets/Sandboxed/SPersistOperationWidget.cpp b/Engine/Plugins/Developer/Sandbox/SandboxedEditing/Source/SandboxedEditing/Private/Features/Browser/Widgets/Sandboxed/SPersistOperationWidget.cpp
--- a/Engine/Plugins/Developer/Sandbox/SandboxedEditing/Source/SandboxedEditing/Private/Features/Browser/Widgets/Sandboxed/SPersistOperationWidget.cpp
+++ b/Engine/Plugins/Developer/Sandbox/SandboxedEditing/Source/SandboxedEditing/Private/Features/Browser/Widgets/Sandboxed/SPersistOperationWidget.cpp
@@ -6,7 +6,7 @@
 #include "Features/Browser/ViewModels/FileState/FileStateItemUtils.h"
 #include "Features/Browser/ViewModels/FileState/Models/StaticFileStateViewModel.h"
 #include "Features/Browser/ViewModels/Persist/PersistOperationViewModel.h"
-#include "FileState/SFileStateListView.h"
+#include "Features/Browser/Widgets/Shared/FileState/SFileStateListView.h"
 #include "SPositiveActionButton.h"
 #include "Styling/AppStyle.h"
 #include "Types/GatheredFileChanges.h"
diff --git a/Engine/Plugins/Developer/Sandbox/SandboxedEditing/Source/SandboxedEditing/Private/Features/Browser/Widgets/Shared/FileState/Columns/ActionFileStateColumn.cpp b/Engine/Plugins/Developer/Sandbox/SandboxedEditing/Source/SandboxedEditing/Private/Features/Browser/Widgets/Shared/FileState/Columns/ActionFileStateColumn.cpp
new file mode 100644
--- /dev/null
+++ b/Engine/Plugins/Developer/Sandbox/SandboxedEditing/Source/SandboxedEditing/Private/Features/Browser/Widgets/Shared/FileState/Columns/ActionFileStateColumn.cpp
@@ -0,0 +1,71 @@
+﻿// Copyright Epic Games, Inc. All Rights Reserved.
+
+#include "ActionFileStateColumn.h"
+
+#include "SandboxedEditingStyle.h"
+#include "Features/Browser/ViewModels/FileState/FileStateColumns.h"
+#include "Features/Browser/ViewModels/FileState/FileStateItem.h"
+#include "RevisionControlStyle/RevisionControlStyle.h"
+#include "Types/SandboxFileChange.h"
+#include "Widgets/Images/SImage.h"
+#include "Widgets/Layout/SBox.h"
+
+#define LOCTEXT_NAMESPACE "FActionFileStateColumn"
+
+namespace UE::SandboxedEditing
+{
+namespace FileActionDetail
+{
+static const FSlateBrush* GetActionBrush(FileSandboxCore::ESandboxFileChange InChange)
+{
+	switch (InChange)
+	{
+	case FileSandboxCore::ESandboxFileChange::Added: return FRevisionControlStyleManager::Get().GetBrush("RevisionControl.OpenForAdd");
+	case FileSandboxCore::ESandboxFileChange::Removed: return FRevisionControlStyleManager::Get().GetBrush("RevisionControl.MarkedForDelete");
+	case FileSandboxCore::ESandboxFileChange::Edited: return FRevisionControlStyleManager::Get().GetBrush("RevisionControl.CheckedOut");
+	case FileSandboxCore::ESandboxFileChange::None: [[fallthrough]];
+	default: return nullptr;
+	}
+}
+
+static FText GetActionToolTipText(FileSandboxCore::ESandboxFileChange InChange)
+{
+	switch (InChange)
+	{
+	case FileSandboxCore::ESandboxFileChange::Added: return LOCTEXT("Added", "Marked for add");
+	case FileSandboxCore::ESandboxFileChange::Removed: return LOCTEXT("Removed", "Marked for delete");
+	case FileSandboxCore::ESandboxFileChange::Edited: return LOCTEXT("Edited", "File is edited");
+	case FileSandboxCore::ESandboxFileChange::None: [[fallthrough]];
+	default: return FText::GetEmpty();
+	}
+}
+}
+
+SHeaderRow::FColumn::FArguments FActionFileStateColumn::MakeColumnArguments()
+{
+	return SHeaderRow::FColumn::FArguments()
+		.ColumnId(FileActionColumn)
+		.DefaultLabel(FText::GetEmpty())
+		.ToolTipText(LOCTEXT("ToolTip", "The action that was performed on the file"))
+		.FixedWidth(FSandboxedEditingStyle::Get().GetFloat("SandboxedEditing.FileActions.FileActionColumn.FixedWidth"))
+		.ShouldGenerateWidget(true)
+		[
+			SNullWidget::NullWidget
+		];
+}
+
+TSharedRef<SWidget> FActionFileStateColumn::MakeColumnWidget(const FMakeFileStateColumnWidgetArgs& InArgs)
+{
+	return SNew(SBox)
+		.HAlign(HAlign_Center)
+		.VAlign(VAlign_Center)
+		.Padding(FSandboxedEditingStyle::Get().GetMargin("SandboxedEditing.FileActions.FileActionColumn.Padding"))
+		[
+			SNew(SImage)
+			.Image(FileActionDetail::GetActionBrush(InArgs.RowData->Action))
+			.ToolTipText(FileActionDetail::GetActionToolTipText(InArgs.RowData->Action))
+		];
+}
+}
+
+#undef LOCTEXT_NAMESPACE
\ No newline at end of file
diff --git a/Engine/Plugins/Developer/Sandbox/SandboxedEditing/Source/SandboxedEditing/Private/Features/Browser/Widgets/Shared/FileState/Columns/ActionFileStateColumn.h b/Engine/Plugins/Developer/Sandbox/SandboxedEditing/Source/SandboxedEditing/Private/Features/Browser/Widgets/Shared/FileState/Columns/ActionFileStateColumn.h
new file mode 100644
--- /dev/null
+++ b/Engine/Plugins/Developer/Sandbox/SandboxedEditing/Source/SandboxedEditing/Private/Features/Browser/Widgets/Shared/FileState/Columns/ActionFileStateColumn.h
@@ -0,0 +1,20 @@
+﻿// Copyright Epic Games, Inc. All Rights Reserved.
+
+#pragma once
+
+#include "Features/Browser/ViewModels/FileState/IFileStateColumnBehavior.h"
+#include "Features/Browser/ViewModels/FileState/IFileStateColumnWidgetFactory.h"
+
+namespace UE::SandboxedEditing
+{
+class FActionFileStateColumn : public IFileStateColumnBehavior, public IFileStateColumnWidgetFactory
+{
+public:
+	
+	//~ Begin IFileStateWidgetFactory Interface
+	virtual SHeaderRow::FColumn::FArguments MakeColumnArguments() override;
+	virtual TSharedRef<SWidget> MakeColumnWidget(const FMakeFileStateColumnWidgetArgs& InArgs) override;
+	//~ End IFileStateWidgetFactory Interface
+};
+}
+
diff --git a/Engine/Plugins/Developer/Sandbox/SandboxedEditing/Source/SandboxedEditing/Private/Features/Browser/Widgets/Shared/FileState/Columns/PathFileStateColumn.cpp b/Engine/Plugins/Developer/Sandbox/SandboxedEditing/Source/SandboxedEditing/Private/Features/Browser/Widgets/Shared/FileState/Columns/PathFileStateColumn.cpp
new file mode 100644
--- /dev/null
+++ b/Engine/Plugins/Developer/Sandbox/SandboxedEditing/Source/SandboxedEditing/Private/Features/Browser/Widgets/Shared/FileState/Columns/PathFileStateColumn.cpp
@@ -0,0 +1,39 @@
+﻿// Copyright Epic Games, Inc. All Rights Reserved.
+
+#include "PathFileStateColumn.h"
+
+#include "Features/Browser/ViewModels/FileState/FileStateColumns.h"
+#include "Features/Browser/ViewModels/FileState/FileStateItem.h"
+#include "SandboxedEditingStyle.h"
+#include "Widgets/Layout/SBox.h"
+#include "Widgets/Text/STextBlock.h"
+
+#define LOCTEXT_NAMESPACE "FPathFileStateColumn"
+
+namespace UE::SandboxedEditing
+{
+SHeaderRow::FColumn::FArguments FPathFileStateColumn::MakeColumnArguments()
+{
+	return SHeaderRow::FColumn::FArguments()
+		.ColumnId(FilePathColumn)
+		.FillWidth(FSandboxedEditingStyle::Get().GetFloat("SandboxedEditing.FileActions.PathColumn.FillWidth"))
+		.DefaultLabel(LOCTEXT("Name.Label", "Path"))
+		.ToolTipText(LOCTEXT("Name.Description", "The absolute file path to the file"))
+		.ShouldGenerateWidget(true);
+}
+
+TSharedRef<SWidget> FPathFileStateColumn::MakeColumnWidget(const FMakeFileStateColumnWidgetArgs& InArgs)
+{
+	const FText AbsPathAsText = FText::AsCultureInvariant(FPaths::ConvertRelativePathToFull(InArgs.RowData->NonSandboxFile));
+	return SNew(SBox)
+		.VAlign(VAlign_Center)
+		.Padding(FSandboxedEditingStyle::Get().GetMargin("SandboxedEditing.FileActions.PathColumn.Padding"))
+		[
+			SNew(STextBlock)
+			.Text(AbsPathAsText)
+			.OverflowPolicy(ETextOverflowPolicy::MiddleEllipsis)
+		];
+}
+}
+
+#undef LOCTEXT_NAMESPACE
\ No newline at end of file
diff --git a/Engine/Plugins/Developer/Sandbox/SandboxedEditing/Source/SandboxedEditing/Private/Features/Browser/Widgets/Shared/FileState/Columns/PathFileStateColumn.h b/Engine/Plugins/Developer/Sandbox/SandboxedEditing/Source/SandboxedEditing/Private/Features/Browser/Widgets/Shared/FileState/Columns/PathFileStateColumn.h
new file mode 100644
--- /dev/null
+++ b/Engine/Plugins/Developer/Sandbox/SandboxedEditing/Source/SandboxedEditing/Private/Features/Browser/Widgets/Shared/FileState/Columns/PathFileStateColumn.h
@@ -0,0 +1,21 @@
+﻿// Copyright Epic Games, Inc. All Rights Reserved.
+
+#pragma once
+
+#include "Features/Browser/ViewModels/FileState/IFileStateColumnBehavior.h"
+#include "Features/Browser/ViewModels/FileState/IFileStateColumnWidgetFactory.h"
+
+namespace UE::SandboxedEditing
+{
+/** Displays the file path to the file. */
+class FPathFileStateColumn : public IFileStateColumnBehavior, public IFileStateColumnWidgetFactory
+{
+public:
+	
+	//~ Begin IFileStateWidgetFactory Interface
+	virtual SHeaderRow::FColumn::FArguments MakeColumnArguments() override;
+	virtual TSharedRef<SWidget> MakeColumnWidget(const FMakeFileStateColumnWidgetArgs& InArgs) override;
+	//~ End IFileStateWidgetFactory Interface
+};
+}
+
diff --git a/Engine/Plugins/Developer/Sandbox/SandboxedEditing/Source/SandboxedEditing/Private/Features/Browser/Widgets/Shared/FileState/Columns/PersistFileStateColumn.cpp b/Engine/Plugins/Developer/Sandbox/SandboxedEditing/Source/SandboxedEditing/Private/Features/Browser/Widgets/Shared/FileState/Columns/PersistFileStateColumn.cpp
new file mode 100644
--- /dev/null
+++ b/Engine/Plugins/Developer/Sandbox/SandboxedEditing/Source/SandboxedEditing/Private/Features/Browser/Widgets/Shared/FileState/Columns/PersistFileStateColumn.cpp
@@ -0,0 +1,76 @@
+﻿// Copyright Epic Games, Inc. All Rights Reserved.
+
+#include "PersistFileStateColumn.h"
+
+#include "SandboxedEditingStyle.h"
+#include "Features/Browser/ViewModels/FileState/FileStateColumns.h"
+#include "Features/Browser/ViewModels/FileState/FileStateItem.h"
+#include "Features/Browser/ViewModels/Persist/PersistOperationViewModel.h"
+#include "Widgets/Input/SCheckBox.h"
+#include "Widgets/Layout/SBox.h"
+
+#define LOCTEXT_NAMESPACE "FPersistFileStateColumn"
+
+namespace UE::SandboxedEditing
+{
+FPersistFileStateColumn::FPersistFileStateColumn(
+const TSharedRef<FPersistOperationViewModel>& InPersistViewModel
+	)
+	: PersistViewModel(InPersistViewModel)
+{}
+
+SHeaderRow::FColumn::FArguments FPersistFileStateColumn::MakeColumnArguments()
+{
+	return SHeaderRow::FColumn::FArguments()
+		.ColumnId(PersistCheckboxColumn)
+		.FixedWidth(FSandboxedEditingStyle::Get().GetFloat("SandboxedEditing.FileActions.PersistCheckbox.FixedWidth"))
+		.ShouldGenerateWidget(true)
+		[
+			SNew(SCheckBox)
+			.IsChecked_Raw(this, &FPersistFileStateColumn::GetToggleRootSelectedState)
+			.OnCheckStateChanged_Raw(this, &FPersistFileStateColumn::OnToggleRootSelectedCheckBox)
+		];
+}
+
+TSharedRef<SWidget> FPersistFileStateColumn::MakeColumnWidget(const FMakeFileStateColumnWidgetArgs& InArgs)
+{
+	const int32 Index = PersistViewModel->GetPersistableFiles().NonSandboxPaths.IndexOfByKey(InArgs.RowData->NonSandboxFile);
+	if (Index == INDEX_NONE)
+	{
+		return SNullWidget::NullWidget;
+	}
+	
+	return SNew(SBox)
+		.HAlign(HAlign_Center)
+		.Padding(FSandboxedEditingStyle::Get().GetMargin("SandboxedEditing.FileActions.PersistCheckbox.Padding"))
+		[
+			SNew(SCheckBox)
+			.IsChecked_Raw(this, &FPersistFileStateColumn::GetItemCheckBoxState, Index)
+			.OnCheckStateChanged_Raw(this, &FPersistFileStateColumn::SetItemCheckBoxState, Index)
+		];
+}
+
+ECheckBoxState FPersistFileStateColumn::GetToggleRootSelectedState() const
+{
+	return PersistViewModel->AreAllFilesMarkedForPersist() ? ECheckBoxState::Checked : ECheckBoxState::Unchecked;
+}
+
+void FPersistFileStateColumn::OnToggleRootSelectedCheckBox(ECheckBoxState InNewState) const
+{
+	PersistViewModel->SetAllFilesPersisted(InNewState == ECheckBoxState::Checked);
+}
+
+ECheckBoxState FPersistFileStateColumn::GetItemCheckBoxState(int32 InIndex) const
+{
+	return PersistViewModel->IsFileMarkedForPersist(InIndex) 
+		? ECheckBoxState::Checked 
+		: ECheckBoxState::Unchecked;
+}
+
+void FPersistFileStateColumn::SetItemCheckBoxState(ECheckBoxState InNewState, int32 InIndex) const
+{
+	PersistViewModel->SetFilePersisted(InIndex, InNewState == ECheckBoxState::Checked);
+}
+}
+
+#undef LOCTEXT_NAMESPACE
\ No newline at end of file
diff --git a/Engine/Plugins/Developer/Sandbox/SandboxedEditing/Source/SandboxedEditing/Private/Features/Browser/Widgets/Shared/FileState/Columns/PersistFileStateColumn.h b/Engine/Plugins/Developer/Sandbox/SandboxedEditing/Source/SandboxedEditing/Private/Features/Browser/Widgets/Shared/FileState/Columns/PersistFileStateColumn.h
new file mode 100644
--- /dev/null
+++ b/Engine/Plugins/Developer/Sandbox/SandboxedEditing/Source/SandboxedEditing/Private/Features/Browser/Widgets/Shared/FileState/Columns/PersistFileStateColumn.h
@@ -0,0 +1,35 @@
+﻿// Copyright Epic Games, Inc. All Rights Reserved.
+
+#pragma once
+
+#include "Features/Browser/ViewModels/FileState/IFileStateColumnBehavior.h"
+#include "Features/Browser/ViewModels/FileState/IFileStateColumnWidgetFactory.h"
+
+namespace UE::SandboxedEditing
+{
+class FPersistOperationViewModel;
+
+class FPersistFileStateColumn : public IFileStateColumnBehavior, public IFileStateColumnWidgetFactory
+{
+public:
+	
+	explicit FPersistFileStateColumn(const TSharedRef<FPersistOperationViewModel>& InPersistViewModel);
+
+	//~ Begin IFileStateWidgetFactory Interface
+	virtual SHeaderRow::FColumn::FArguments MakeColumnArguments() override;
+	virtual TSharedRef<SWidget> MakeColumnWidget(const FMakeFileStateColumnWidgetArgs& InArgs) override;
+	//~ End IFileStateWidgetFactory Interface
+	
+private:
+	
+	/** Used to set whether an item should be persisted */
+	const TSharedRef<FPersistOperationViewModel> PersistViewModel;
+
+	ECheckBoxState GetToggleRootSelectedState() const;
+	void OnToggleRootSelectedCheckBox(ECheckBoxState InNewState) const;
+
+	ECheckBoxState GetItemCheckBoxState(int32 InIndex) const;
+	void SetItemCheckBoxState(ECheckBoxState InNewState, int32 InIndex) const;
+};
+}
+
diff --git a/Engine/Plugins/Developer/Sandbox/SandboxedEditing/Source/SandboxedEditing/Private/Features/Browser/Widgets/Shared/FileState/SFileStateListRow.cpp b/Engine/Plugins/Developer/Sandbox/SandboxedEditing/Source/SandboxedEditing/Private/Features/Browser/Widgets/Shared/FileState/SFileStateListRow.cpp
new file mode 100644
--- /dev/null
+++ b/Engine/Plugins/Developer/Sandbox/SandboxedEditing/Source/SandboxedEditing/Private/Features/Browser/Widgets/Shared/FileState/SFileStateListRow.cpp
@@ -0,0 +1,32 @@
+﻿// Copyright Epic Games, Inc. All Rights Reserved.
+
+#include "SFileStateListRow.h"
+
+#include "Features/Browser/ViewModels/FileState/IFileStateColumnWidgetFactory.h"
+
+namespace UE::SandboxedEditing
+{
+void SFileStateListRow::Construct(const FArguments& InArgs, const TSharedRef<FFileStateItem>& InItem, const TSharedRef<STableViewBase>& InOwnerTable)
+{
+	Item = InItem;
+	ColumnFactoryMap = InArgs._ColumnFactories;
+	
+	SMultiColumnTableRow::Construct(
+		SMultiColumnTableRow::FArguments().Padding(InArgs._Padding),
+		InOwnerTable
+		);
+}
+
+TSharedRef<SWidget> SFileStateListRow::GenerateWidgetForColumn(const FName& InColumnName)
+{
+	if (TSharedRef<IFileStateColumnWidgetFactory>* Factory = ColumnFactoryMap.Find(InColumnName))
+	{
+		return Factory->Get().MakeColumnWidget(
+			FMakeFileStateColumnWidgetArgs(Item)
+		);
+	}
+	
+	checkNoEntry();
+	return SNullWidget::NullWidget;
+}
+}
diff --git a/Engine/Plugins/Developer/Sandbox/SandboxedEditing/Source/SandboxedEditing/Private/Features/Browser/Widgets/Shared/FileState/SFileStateListRow.h b/Engine/Plugins/Developer/Sandbox/SandboxedEditing/Source/SandboxedEditing/Private/Features/Browser/Widgets/Shared/FileState/SFileStateListRow.h
new file mode 100644
--- /dev/null
+++ b/Engine/Plugins/Developer/Sandbox/SandboxedEditing/Source/SandboxedEditing/Private/Features/Browser/Widgets/Shared/FileState/SFileStateListRow.h
@@ -0,0 +1,41 @@
+﻿// Copyright Epic Games, Inc. All Rights Reserved.
+
+#pragma once
+
+#include "Features/Browser/ViewModels/FileState/FileStateColumns.h"
+#include "Features/Browser/ViewModels/FileState/FileStateItem.h"
+#include "Features/Browser/ViewModels/FileState/IFileStateColumnWidgetFactory.h"
+#include "Features/Browser/ViewModels/List/ISandboxColumnWidgetFactory.h"
+#include "Widgets/DeclarativeSyntaxSupport.h"
+#include "Widgets/Views/SListView.h"
+#include "Widgets/Views/STableRow.h"
+
+namespace UE::SandboxedEditing
+{
+class SFileStateListRow : public SMultiColumnTableRow<TSharedPtr<FFileStateItem>>
+{
+public:
+	
+	SLATE_BEGIN_ARGS(SFileStateListRow){}
+		/** Used to build the columns. */
+		SLATE_ARGUMENT(FFileStateColumnFactoryMap, ColumnFactories)
+		/** Padding for the row */
+		SLATE_ARGUMENT(FMargin, Padding)
+	SLATE_END_ARGS()
+	
+	void Construct(const FArguments& InArgs, const TSharedRef<FFileStateItem>& InItem, const TSharedRef<STableViewBase>& InOwnerTable);
+
+	//~ Begin SMultiColumnTableRow Interface
+	virtual TSharedRef<SWidget> GenerateWidgetForColumn(const FName& InColumnName) override;
+	//~ End SMultiColumnTableRow Interface
+	
+private:
+	
+	/** The item that is being displayed by this row. */
+	TSharedPtr<FFileStateItem> Item;
+	
+	/** Used to build the columns. */
+	FFileStateColumnFactoryMap ColumnFactoryMap;
+};
+}
+
diff --git a/Engine/Plugins/Developer/Sandbox/SandboxedEditing/Source/SandboxedEditing/Private/Features/Browser/Widgets/Shared/FileState/SFileStateListView.cpp b/Engine/Plugins/Developer/Sandbox/SandboxedEditing/Source/SandboxedEditing/Private/Features/Browser/Widgets/Shared/FileState/SFileStateListView.cpp
new file mode 100644
--- /dev/null
+++ b/Engine/Plugins/Developer/Sandbox/SandboxedEditing/Source/SandboxedEditing/Private/Features/Browser/Widgets/Shared/FileState/SFileStateListView.cpp
@@ -0,0 +1,54 @@
+﻿// Copyright Epic Games, Inc. All Rights Reserved.
+
+#include "SFileStateListView.h"
+
+#include "SandboxedEditingStyle.h"
+#include "SFileStateListRow.h"
+#include "Features/Browser/ViewModels/FileState/IFileStateColumnWidgetFactory.h"
+#include "Features/Browser/ViewModels/FileState/IFileStateViewModel.h"
+#include "Widgets/Views/SHeaderRow.h"
+
+namespace UE::SandboxedEditing
+{
+void SFileStateListView::Construct(const FArguments& InArgs, const TSharedRef<IFileStateViewModel>& InViewModel)
+{
+	ViewModel = InViewModel;
+	ColumnFactories = InArgs._ColumnFactories;
+	
+	ChildSlot
+	[
+		SAssignNew(ListView, SListView<TSharedPtr<FFileStateItem>>)
+		.ListItemsSource(&ViewModel->GetItems())
+		.OnGenerateRow(this, &SFileStateListView::OnGenerateRow)
+		.SelectionMode(ESelectionMode::None)
+		.HeaderRow(MakeHeaderRow())
+	];
+	
+	ViewModel->OnItemsChanged().AddSP(ListView.ToSharedRef(), &SListView<TSharedPtr<FFileStateItem>>::RequestListRefresh);
+}
+
+TSharedRef<SHeaderRow> SFileStateListView::MakeHeaderRow()
+{
+	const TSharedRef<SHeaderRow> HeaderRow = SNew(SHeaderRow);
+		
+	for (FName ColumnId : ViewModel->GetDisplayedColumns())
+	{
+		if (const TSharedRef<IFileStateColumnWidgetFactory>* Factory = ColumnFactories.Find(ColumnId))
+		{
+			HeaderRow->AddColumn(Factory->Get().MakeColumnArguments());
+		}
+	}
+	
+	return HeaderRow;
+}
+
+TSharedRef<ITableRow> SFileStateListView::OnGenerateRow(TSharedPtr<FFileStateItem> InItem, const TSharedRef<STableViewBase>& InOwnerTable) const
+{
+	check(InItem);
+	const FMargin Padding = FSandboxedEditingStyle::Get().GetMargin("SandboxedEditing.FileActions.RowPadding");
+	
+	return SNew(SFileStateListRow, InItem.ToSharedRef(), InOwnerTable)
+		.ColumnFactories(ColumnFactories)
+		.Padding(Padding);
+}
+}
diff --git a/Engine/Plugins/Developer/Sandbox/SandboxedEditing/Source/SandboxedEditing/Private/Features/Browser/Widgets/Shared/FileState/SFileStateListView.h b/Engine/Plugins/Developer/Sandbox/SandboxedEditing/Source/SandboxedEditing/Private/Features/Browser/Widgets/Shared/FileState/SFileStateListView.h
new file mode 100644
--- /dev/null
+++ b/Engine/Plugins/Developer/Sandbox/SandboxedEditing/Source/SandboxedEditing/Private/Features/Browser/Widgets/Shared/FileState/SFileStateListView.h
@@ -0,0 +1,48 @@
+﻿// Copyright Epic Games, Inc. All Rights Reserved.
+
+#pragma once
+
+#include "Features/Browser/ViewModels/FileState/FileStateColumns.h"
+#include "Widgets/DeclarativeSyntaxSupport.h"
+#include "Widgets/SCompoundWidget.h"
+#include "Widgets/Views/SListView.h"
+
+class STableViewBase;
+class ITableRow;
+class SHeaderRow;
+
+namespace UE::SandboxedEditing
+{
+class IFileStateViewModel;
+
+/** Displays a list view FFileStateItem provided by IFileStateViewModel. */
+class SFileStateListView : public SCompoundWidget
+{
+public:
+	
+	SLATE_BEGIN_ARGS(SFileStateListView) {}
+		/** Used to build the columns. */
+		SLATE_ARGUMENT(FFileStateColumnFactoryMap, ColumnFactories)
+	SLATE_END_ARGS()
+	
+	void Construct(const FArguments& InArgs, const TSharedRef<IFileStateViewModel>& InViewModel);
+	
+private:
+	
+	/** The view model being displayed */
+	TSharedPtr<IFileStateViewModel> ViewModel;
+	
+	/** The list view displaying the items. */
+	TSharedPtr<SListView<TSharedPtr<FFileStateItem>>> ListView;
+	
+	/** Used to build the columns. */
+	TMap<FName, TSharedRef<IFileStateColumnWidgetFactory>> ColumnFactories;
+	
+	/** @return The header row to display at the top of the table */
+	TSharedRef<SHeaderRow> MakeHeaderRow();
+	
+	/** @return Widget for the items */
+	TSharedRef<ITableRow> OnGenerateRow(TSharedPtr<FFileStateItem> InItem, const TSharedRef<STableViewBase>& InOwnerTable) const;
+};
+}
+
diff --git a/Engine/Plugins/Developer/Sandbox/SandboxedEditing/Source/SandboxedEditing/Private/Features/Browser/Widgets/Unsandboxed/SBrowserSandboxDetails.cpp b/Engine/Plugins/Developer/Sandbox/SandboxedEditing/Source/SandboxedEditing/Private/Features/Browser/Widgets/Unsandboxed/SBrowserSandboxDetails.cpp
--- a/Engine/Plugins/Developer/Sandbox/SandboxedEditing/Source/SandboxedEditing/Private/Features/Browser/Widgets/Unsandboxed/SBrowserSandboxDetails.cpp
+++ b/Engine/Plugins/Developer/Sandbox/SandboxedEditing/Source/SandboxedEditing/Private/Features/Browser/Widgets/Unsandboxed/SBrowserSandboxDetails.cpp
@@ -5,7 +5,7 @@
 #include "Components/VerticalBox.h"
 #include "Features/Browser/ViewModels/FileState/FileStateColumnRegistry.h"
 #include "Features/Browser/ViewModels/FileState/Models/UnloadedSandboxFileStateViewModel.h"
-#include "Features/Browser/Widgets/Sandboxed/FileState/SFileStateListView.h"
+#include "Features/Browser/Widgets/Shared/FileState/SFileStateListView.h"
 #include "Features/Browser/Widgets/Shared/SSandboxDescription.h"
 #include "Internationalization/Text.h"
 #include "Widgets/Layout/SBox.h"
