Sandboxed Editing: Sandboxes can now be renamed. The sandbox list view also has a right-click context menu.
- Slow double click
- Right-click > Rename (only when one item is selected)
- Renaming displays a notification "Renamed sandbox {0} as {1}".

Right-click menu now also has delete option

#Synced-CL 49159615

diff --git a/Engine/Plugins/Developer/Sandbox/FileSandbox/Source/FileSandboxCore/Private/Utils/SandboxDirectoryUtils.cpp b/Engine/Plugins/Developer/Sandbox/FileSandbox/Source/FileSandboxCore/Private/Utils/SandboxDirectoryUtils.cpp
--- a/Engine/Plugins/Developer/Sandbox/FileSandbox/Source/FileSandboxCore/Private/Utils/SandboxDirectoryUtils.cpp
+++ b/Engine/Plugins/Developer/Sandbox/FileSandbox/Source/FileSandboxCore/Private/Utils/SandboxDirectoryUtils.cpp
@@ -3,6 +3,9 @@
 #include "Utils/SandboxDirectoryUtils.h"
 
 #include "CommandLineUtils.h"
+#include "IFileSandboxCoreModule.h"
+#include "ISandboxInstance.h"
+#include "ISandboxManager.h"
 #include "Data/ManifestData.h"
 #include "Data/SandboxMetaData.h"
 #include "Data/VersionInfo.h"
@@ -83,6 +86,75 @@ TOptional<FFileSandboxCore_SandboxMetaData> LoadMetaData(const FString& InSandbo
 	return {};
 }
 
+bool RenameSandboxWithDirectory(
+	const FString& InNewName, const FString& InOldName, const FString& InBaseDirectory, const FFileSandboxCore_SandboxMetaData* InNewMetaData
+	)
+{
+	if (CanRenameSandboxWithDirectory(InNewName, InOldName, InBaseDirectory) != ESandboxRenameSuitability::Allowed)
+	{
+		return false;
+	}
+	
+	const FString OldSandboxDir = InBaseDirectory / InOldName;
+	TOptional<FFileSandboxCore_SandboxMetaData> NewMetaData = InNewMetaData 
+		? TOptional(*InNewMetaData) : LoadMetaData(*OldSandboxDir);
+	if (!NewMetaData)
+	{
+		return false;
+	}
+	
+	NewMetaData->Name = InNewName;
+	if (!SaveMetaData(*NewMetaData, OldSandboxDir))
+	{
+		return false;
+	}
+	
+	const FString NewDir = InBaseDirectory / InNewName;
+	return IFileManager::Get().Move(*NewDir, *OldSandboxDir);
+}
+
+ESandboxRenameSuitability CanRenameSandboxWithDirectory(const FString& InNewName, const FString& InOldName, const FString& InBaseDirectory)
+{
+	if (InNewName.IsEmpty())
+	{
+		return ESandboxRenameSuitability::EmptyName;
+	}
+	
+	if (InNewName == InOldName)
+	{
+		return ESandboxRenameSuitability::NameIsSame;
+	}
+	
+	const FString OldDir = InBaseDirectory / InOldName;
+	const bool bSandboxExists = IsRootSandboxDirectory(*OldDir);
+	if (!bSandboxExists)
+	{
+		return ESandboxRenameSuitability::SandboxDoesNotExist;
+	}
+	
+	if (!IsAllowedToRenameSandbox(OldDir))
+	{
+		return ESandboxRenameSuitability::ActiveSandbox;
+	}
+	
+	const FString NewDirectory = InBaseDirectory / InNewName;
+	const bool bNewDirectoryIsAvailable = FPaths::ValidatePath(NewDirectory) && !FPaths::DirectoryExists(NewDirectory);
+	if (!bNewDirectoryIsAvailable)
+	{
+		return ESandboxRenameSuitability::InvalidDirectory;
+	}
+	
+	return ESandboxRenameSuitability::Allowed;
+}
+
+bool IsAllowedToRenameSandbox(const FString& InSandboxDirectory)
+{
+	// Renaming would leave the internal sandbox root directory "stale". Let's not implement advanced renaming logic for that for now.
+	ISandboxInstance* ActiveSandbox = IFileSandboxCoreModule::Get().GetSandboxManager().GetActiveSandboxInstance();
+	const bool bIsNotActiveSandbox = !ActiveSandbox || !FPaths::IsSamePath(InSandboxDirectory, ActiveSandbox->GetRootDirectory());
+	return bIsNotActiveSandbox;
+}
+
 FDateTime GetLastModified(const FString& InSandboxDirectory)
 {
 	IPlatformFile& PlatformFile = FPlatformFileManager::Get().GetPlatformFile();
diff --git a/Engine/Plugins/Developer/Sandbox/FileSandbox/Source/FileSandboxCore/Public/Utils/SandboxDirectoryUtils.h b/Engine/Plugins/Developer/Sandbox/FileSandbox/Source/FileSandboxCore/Public/Utils/SandboxDirectoryUtils.h
--- a/Engine/Plugins/Developer/Sandbox/FileSandbox/Source/FileSandboxCore/Public/Utils/SandboxDirectoryUtils.h
+++ b/Engine/Plugins/Developer/Sandbox/FileSandbox/Source/FileSandboxCore/Public/Utils/SandboxDirectoryUtils.h
@@ -40,8 +40,8 @@ UE_API bool IsRootSandboxDirectory(const FString& InRootDirectory);
 /**
  * Writes new metadata for a sandbox.
  * 
- * You're allowed to change the sandbox name. The name of the base directory will remain be changed as well (for consistency).
- * As such, this function will fail if you're trying to rename the sandbox while it is active (for safety & easier implementation reasons).
+ * You are allowed to change only the sandbox name. However, it will not rename the sandbox directory.
+ * To rename the directory, @see RenameSandboxWithDirectory.
  * 
  * @param NewMetaData The new metadata to set for the sandbox.
  * @param InSandboxDirectory Path to the base directory of the sandbox, e.g. what's returned by ForEachSandbox. Either absolute or relative to the process BaseDir().
@@ -56,6 +56,49 @@ UE_API bool SaveMetaData(
  */
 UE_API TOptional<FFileSandboxCore_SandboxMetaData> LoadMetaData(const FString& InSandboxDirectory);
 
+/**
+ * Renames the sandbox located in InBaseDirectory / InOldName. 
+ * 
+ * By default, sandboxes are placed in BaseDir / SandboxName. This function renames the containing directory to the new name.
+ * This also updates the manifest file with the new name.
+ * 
+ * @param InNewName The new name the sandbox should have
+ * @param InOldName The sandbox you want to rename. The sandbox root is located in InBaseDirectory / InOldName directory.
+ * @param InBaseDirectory Base directory search for. Either absolute or relative to the process BaseDir().
+ * @param InNewMetaData Optional additional metadata you want to set. Specifying this will avoid loading the metadata from disk first.
+ * @return Whether the sandbox was successfully renamed
+ */
+UE_API bool RenameSandboxWithDirectory(
+	const FString& InNewName, const FString& InOldName, 
+	const FString& InBaseDirectory = GetBaseSandboxDirectory(),
+	const FFileSandboxCore_SandboxMetaData* InNewMetaData = nullptr
+	);
+
+/** Explains why a sandbox can or cannot be renamed. */
+enum class ESandboxRenameSuitability : uint8
+{
+	Allowed,
+	
+	/** The sandbox does not exist */
+	SandboxDoesNotExist,
+	/** You cannot rename an active sandbox */
+	ActiveSandbox,
+	
+	/** The name is not different from the current name. */
+	NameIsSame,
+	/** Sandbox cannot have this name */
+	InvalidDirectory,
+	/** Name cannot be empty */
+	EmptyName,
+};
+
+/** @return Whether InNewName is a valid new name for the sandbox located in InBaseDirectory / InOldName. */
+UE_API ESandboxRenameSuitability CanRenameSandboxWithDirectory(
+	const FString& InNewName, const FString& InOldName, const FString& InBaseDirectory = GetBaseSandboxDirectory()
+	);
+/** @return Whether the given sandbox can be renamed, i.e. it is not active. */
+UE_API bool IsAllowedToRenameSandbox(const FString& InSandboxDirectory);
+
 
 /** @return When the given sandbox was last modified. Returns FDateTime::MinValue() on failure. */
 UE_API FDateTime GetLastModified(const FString& InSandboxDirectory);
diff --git a/Engine/Plugins/Developer/Sandbox/SandboxedEditing/Source/SandboxedEditing/Private/Features/Browser/BrowserFeature.cpp b/Engine/Plugins/Developer/Sandbox/SandboxedEditing/Source/SandboxedEditing/Private/Features/Browser/BrowserFeature.cpp
--- a/Engine/Plugins/Developer/Sandbox/SandboxedEditing/Source/SandboxedEditing/Private/Features/Browser/BrowserFeature.cpp
+++ b/Engine/Plugins/Developer/Sandbox/SandboxedEditing/Source/SandboxedEditing/Private/Features/Browser/BrowserFeature.cpp
@@ -22,9 +22,12 @@ namespace BrowserDetail
 {
 const FLazyName SandboxedEditingTabName(TEXT("SandboxedEditingTab"));
 
-static void MakeSandboxColumns(TMap<FName, TSharedRef<ISandboxColumnBehavior>>& OutBehaviors, TMap<FName, TSharedRef<ISandboxColumnWidgetFactory>>& OutFactories)
+static void MakeSandboxColumns(
+	const TSharedRef<FSandboxControlsViewModel>& InControlsViewModel,
+	TMap<FName, TSharedRef<ISandboxColumnBehavior>>& OutBehaviors, TMap<FName, TSharedRef<ISandboxColumnWidgetFactory>>& OutFactories
+	)
 {
-	const TSharedRef<FNameSandboxColumn> Name = MakeShared<FNameSandboxColumn>();
+	const TSharedRef<FNameSandboxColumn> Name = MakeShared<FNameSandboxColumn>(InControlsViewModel);
 	const TSharedRef<FVersionSandboxColumn> Version = MakeShared<FVersionSandboxColumn>();
 	const TSharedRef<FLastModifiedSandboxColumn> LastModified = MakeShared<FLastModifiedSandboxColumn>();
 	OutBehaviors = TMap<FName, TSharedRef<ISandboxColumnBehavior>>
@@ -41,10 +44,10 @@ static void MakeSandboxColumns(TMap<FName, TSharedRef<ISandboxColumnBehavior>>&
 	};
 }
 
-static FSandboxColumnRegistry MakeSandboxColumnRegistry()
+static FSandboxColumnRegistry MakeSandboxColumnRegistry(const TSharedRef<FSandboxControlsViewModel>& InControlsViewModel)
 {
 	FSandboxColumnRegistry Registry;
-	MakeSandboxColumns(Registry.ColumnBehaviors, Registry.ColumnFactories);
+	MakeSandboxColumns(InControlsViewModel, Registry.ColumnBehaviors, Registry.ColumnFactories);
 	return Registry;
 }
 }
@@ -53,12 +56,20 @@ FBrowserFeature::FBrowserFeature(
 	const TSharedRef<FSandboxSystemModel>& InSandboxModel, 
 	const TSharedRef<FileSandboxUI::IExternalSandboxActiveViewModel>& InExternalSandboxViewModel
 	)
+	: FBrowserFeature(InSandboxModel, InExternalSandboxViewModel, MakeShared<FSandboxControlsViewModel>(InSandboxModel))
+{}
+
+FBrowserFeature::FBrowserFeature(
+	const TSharedRef<FSandboxSystemModel>& InSandboxModel,
+	const TSharedRef<FileSandboxUI::IExternalSandboxActiveViewModel>& InExternalSandboxViewModel,
+	const TSharedRef<FSandboxControlsViewModel>& InControlsViewModel
+	)
 	: SandboxModel(InSandboxModel)
-	, SandboxColumns(BrowserDetail::MakeSandboxColumnRegistry())
+	, SandboxColumns(BrowserDetail::MakeSandboxColumnRegistry(InControlsViewModel))
 	, BrowserFileActionColumns(GetColumnsForBrowser())
 	, ActiveFileActionColumns(GetColumnsForActiveSandbox())
-	, ViewModels(SandboxColumns, BrowserFileActionColumns, ActiveFileActionColumns, SandboxModel, InExternalSandboxViewModel)
-	, Views(ViewModels.PersistViewModel, ViewModels.LeaveViewModel)
+	, ViewModels(SandboxColumns, BrowserFileActionColumns, ActiveFileActionColumns, SandboxModel, InExternalSandboxViewModel, InControlsViewModel)
+	, Views(ViewModels.ControlsViewModel, ViewModels.PersistViewModel, ViewModels.LeaveViewModel)
 	, CommandBindings(MakeShared<FBrowserCommandBindings>(ViewModels))
 {
 	RegisterTabSpawner();
@@ -74,6 +85,8 @@ void FBrowserFeature::SummonUI() const
 	FGlobalTabmanager::Get()->TryInvokeTab(BrowserDetail::SandboxedEditingTabName.Resolve());
 }
 
+
+
 void FBrowserFeature::RegisterTabSpawner()
 {
 	FTabSpawnerEntry& BrowserSpawnerEntry = FGlobalTabmanager::Get()->RegisterNomadTabSpawner(
diff --git a/Engine/Plugins/Developer/Sandbox/SandboxedEditing/Source/SandboxedEditing/Private/Features/Browser/BrowserFeature.h b/Engine/Plugins/Developer/Sandbox/SandboxedEditing/Source/SandboxedEditing/Private/Features/Browser/BrowserFeature.h
--- a/Engine/Plugins/Developer/Sandbox/SandboxedEditing/Source/SandboxedEditing/Private/Features/Browser/BrowserFeature.h
+++ b/Engine/Plugins/Developer/Sandbox/SandboxedEditing/Source/SandboxedEditing/Private/Features/Browser/BrowserFeature.h
@@ -59,6 +59,12 @@ private:
 	/** Binds user commands to models.*/
 	const TSharedRef<FBrowserCommandBindings> CommandBindings;
 	
+	explicit FBrowserFeature(
+		const TSharedRef<FSandboxSystemModel>& InSandboxModel, 
+		const TSharedRef<FileSandboxUI::IExternalSandboxActiveViewModel>& InExternalSandboxViewModel,
+		const TSharedRef<FSandboxControlsViewModel>& InControlsViewModel
+		);
+	
 	void RegisterTabSpawner();
 	void UnregisterTabSpawner();
 	TSharedRef<SDockTab> SpawnBrowserTab(const FSpawnTabArgs& InArgs);
diff --git a/Engine/Plugins/Developer/Sandbox/SandboxedEditing/Source/SandboxedEditing/Private/Features/Browser/Commands/BrowserCommandBindings.cpp b/Engine/Plugins/Developer/Sandbox/SandboxedEditing/Source/SandboxedEditing/Private/Features/Browser/Commands/BrowserCommandBindings.cpp
--- a/Engine/Plugins/Developer/Sandbox/SandboxedEditing/Source/SandboxedEditing/Private/Features/Browser/Commands/BrowserCommandBindings.cpp
+++ b/Engine/Plugins/Developer/Sandbox/SandboxedEditing/Source/SandboxedEditing/Private/Features/Browser/Commands/BrowserCommandBindings.cpp
@@ -41,10 +41,7 @@ FBrowserCommandBindings::~FBrowserCommandBindings()
 
 void FBrowserCommandBindings::HandleCancel() const
 {
-	if (FSandboxCreationWorkflow* CreationWorkflow = ControlsViewModel->GetCurrentCreationWorkflow())
-	{
-		CreationWorkflow->Cancel();
-	}
+	ControlsViewModel->CancelWorkflows();
 }
 
 void FBrowserCommandBindings::HandleLeaveSandbox()
diff --git a/Engine/Plugins/Developer/Sandbox/SandboxedEditing/Source/SandboxedEditing/Private/Features/Browser/ViewModels/BrowserViewModels.h b/Engine/Plugins/Developer/Sandbox/SandboxedEditing/Source/SandboxedEditing/Private/Features/Browser/ViewModels/BrowserViewModels.h
--- a/Engine/Plugins/Developer/Sandbox/SandboxedEditing/Source/SandboxedEditing/Private/Features/Browser/ViewModels/BrowserViewModels.h
+++ b/Engine/Plugins/Developer/Sandbox/SandboxedEditing/Source/SandboxedEditing/Private/Features/Browser/ViewModels/BrowserViewModels.h
@@ -58,13 +58,14 @@ struct FBrowserViewModels
 		const FFileStateColumnRegistry& InBrowserFileActionColumns,
 		const FFileStateColumnRegistry& InActiveFileActionColumns,
 		const TSharedRef<FSandboxSystemModel>& InModel,
-		const TSharedRef<FileSandboxUI::IExternalSandboxActiveViewModel>& InExternalSandboxViewModel
+		const TSharedRef<FileSandboxUI::IExternalSandboxActiveViewModel>& InExternalSandboxViewModel,
+		const TSharedRef<FSandboxControlsViewModel>& InControlsViewModel
 		)
 		: FilterSandboxViewModel(MakeShared<FFilterSandboxViewModel>(InSandboxColumns.GetBehaviorArray()))
 		, ListViewModel(MakeShared<FSandboxListViewModel>(InModel, FilterSandboxViewModel, InSandboxColumns.ColumnBehaviors))
 		, UnloadedSandboxFileStateViewModel(MakeShared<FUnloadedSandboxFileStateViewModel>(InModel, InBrowserFileActionColumns.ColumnBehaviors))
 		, ActiveSandboxFileStateViewModel(MakeShared<FActiveSandboxFileStateViewModel>(InModel, InActiveFileActionColumns.ColumnBehaviors))
-		, ControlsViewModel(MakeShared<FSandboxControlsViewModel>(InModel))
+		, ControlsViewModel(InControlsViewModel)
 		, MetaDataViewModel(MakeShared<FSandboxMetaDataViewModel>(InModel))
 		, PersistViewModel(MakeShared<FPersistSandboxViewModel>(InModel))
 		, LeaveViewModel(MakeShared<FLeaveSandboxViewModel>(InModel, PersistViewModel))
diff --git a/Engine/Plugins/Developer/Sandbox/SandboxedEditing/Source/SandboxedEditing/Private/Features/Browser/ViewModels/RenameSandboxWorkflow.cpp b/Engine/Plugins/Developer/Sandbox/SandboxedEditing/Source/SandboxedEditing/Private/Features/Browser/ViewModels/RenameSandboxWorkflow.cpp
new file mode 100644
--- /dev/null
+++ b/Engine/Plugins/Developer/Sandbox/SandboxedEditing/Source/SandboxedEditing/Private/Features/Browser/ViewModels/RenameSandboxWorkflow.cpp
@@ -0,0 +1,44 @@
+﻿// Copyright Epic Games, Inc. All Rights Reserved.
+
+#include "RenameSandboxWorkflow.h"
+
+#include "Framework/Models/SandboxInfo.h"
+#include "Framework/Models/SandboxSystemModel.h"
+
+namespace UE::SandboxedEditing
+{
+FRenameSandboxWorkflow::FRenameSandboxWorkflow(
+	FString InRenamedRoot, 
+	const TSharedRef<FSandboxSystemModel>& InModel,
+	FSimpleDelegate InOnWorkflowEndedDelegate
+	)
+	: Model(InModel)
+	, RenamedSandboxRoot(MoveTemp(InRenamedRoot))
+	, InitialName(Model->GetSandboxInfo(RenamedSandboxRoot)->Name)
+	, OnWorkflowEndedDelegate(MoveTemp(InOnWorkflowEndedDelegate))
+{}
+
+void FRenameSandboxWorkflow::Confirm()
+{
+	if (IsNameValid() && Model->RenameSandbox(RenamedSandboxRoot, NewName))
+	{
+		OnSandboxRenamedDelegate.Broadcast(InitialName, NewName);
+		OnWorkflowEndedDelegate.ExecuteIfBound();
+	}
+}
+
+void FRenameSandboxWorkflow::Cancel()
+{
+	OnWorkflowEndedDelegate.ExecuteIfBound();
+}
+
+bool FRenameSandboxWorkflow::IsNameValid(FText* OutErrorText) const
+{
+	if (!bHasChangedName)
+	{
+		return true;
+	}
+	
+	return Model->CanRenameSandbox(RenamedSandboxRoot, NewName, OutErrorText);
+}
+}
diff --git a/Engine/Plugins/Developer/Sandbox/SandboxedEditing/Source/SandboxedEditing/Private/Features/Browser/ViewModels/RenameSandboxWorkflow.h b/Engine/Plugins/Developer/Sandbox/SandboxedEditing/Source/SandboxedEditing/Private/Features/Browser/ViewModels/RenameSandboxWorkflow.h
new file mode 100644
--- /dev/null
+++ b/Engine/Plugins/Developer/Sandbox/SandboxedEditing/Source/SandboxedEditing/Private/Features/Browser/ViewModels/RenameSandboxWorkflow.h
@@ -0,0 +1,67 @@
+﻿// Copyright Epic Games, Inc. All Rights Reserved.
+
+#pragma once
+
+#include "Delegates/Delegate.h"
+#include "Templates/SharedPointer.h"
+#include "Templates/UnrealTemplate.h"
+
+namespace UE::SandboxedEditing
+{
+DECLARE_MULTICAST_DELEGATE_TwoParams(FOnSandboxRenamed, const FString& /*InOldName*/, const FString& /*InNewName*/);
+
+class FSandboxSystemModel;
+
+/** Workflow for renaming a sandbox. */
+class FRenameSandboxWorkflow : public FNoncopyable
+{
+public:
+	
+	explicit FRenameSandboxWorkflow(
+		FString InRenamedRoot,
+		const TSharedRef<FSandboxSystemModel>& InModel, FSimpleDelegate InOnWorkflowEndedDelegate
+		);
+	
+	/** Renames the sandbox if the name is valid, and stops the operation if successful. */
+	void Confirm();
+	/** Cancels the operation. */
+	void Cancel();
+	
+	/** Sets the that the sandbox should be renamed to. */
+	void SetName(FString InName)
+	{
+		bHasChangedName |= InitialName != InName;
+		NewName = InName;
+	}
+	
+	/** @return Whether the current name is valid, and an optional error message if this function returns false. */
+	bool IsNameValid(FText* OutErrorText = nullptr) const;
+	
+	/** @return The root directory of the sandbox being renamed */
+	const FString& GetRenamedSandboxRoot() const { return RenamedSandboxRoot; }
+	
+	/** Invoked after the sandbox has been successfully renamed. */
+	FOnSandboxRenamed& OnSandboxRenamed() { return OnSandboxRenamedDelegate; }
+
+private:  
+	
+	/** The model used to interact with the sandbox system. */
+	const TSharedRef<FSandboxSystemModel> Model;
+	
+	/** The root directory of the sandbox being renamed */
+	const FString RenamedSandboxRoot;
+	/** The name the sandbox had at the beginning. */
+	const FString InitialName;
+	
+	/** Invoked after the sandbox has been successfully renamed. */
+	FOnSandboxRenamed OnSandboxRenamedDelegate;
+	/** Invoked when the workflow ends. */
+	const FSimpleDelegate OnWorkflowEndedDelegate;
+	
+	/** The name to give the sandbox. */
+	FString NewName;
+	
+	bool bHasChangedName = false;
+};
+}
+
diff --git a/Engine/Plugins/Developer/Sandbox/SandboxedEditing/Source/SandboxedEditing/Private/Features/Browser/ViewModels/SandboxControlsViewModel.cpp b/Engine/Plugins/Developer/Sandbox/SandboxedEditing/Source/SandboxedEditing/Private/Features/Browser/ViewModels/SandboxControlsViewModel.cpp
--- a/Engine/Plugins/Developer/Sandbox/SandboxedEditing/Source/SandboxedEditing/Private/Features/Browser/ViewModels/SandboxControlsViewModel.cpp
+++ b/Engine/Plugins/Developer/Sandbox/SandboxedEditing/Source/SandboxedEditing/Private/Features/Browser/ViewModels/SandboxControlsViewModel.cpp
@@ -8,13 +8,26 @@ FSandboxControlsViewModel::FSandboxControlsViewModel(const TSharedRef<FSandboxSy
 	: Model(InModel)
 {}
 
+void FSandboxControlsViewModel::CancelWorkflows()
+{
+	if (CurrentCreationWorkflow)
+	{
+		CurrentCreationWorkflow->Cancel();
+	}
+	
+	if (CurrentRenameWorkflow)
+	{
+		CurrentRenameWorkflow->Cancel();
+	}
+}
+
 void FSandboxControlsViewModel::StartCreationWorkflow()
 {
 	if (CanStartCreationWorkflow())
 	{
 		CurrentCreationWorkflow = MakeUnique<FSandboxCreationWorkflow>(
 			Model,
-			FSimpleDelegate::CreateRaw(this, &FSandboxControlsViewModel::OnWorkflowEnded)
+			FSimpleDelegate::CreateRaw(this, &FSandboxControlsViewModel::HandleCreationWorkflowEnded)
 			);
 		OnCreationWorkflowStartedDelegate.Broadcast(*CurrentCreationWorkflow);
 	}
@@ -22,7 +35,7 @@ void FSandboxControlsViewModel::StartCreationWorkflow()
 
 bool FSandboxControlsViewModel::CanStartCreationWorkflow() const
 {
-	return CurrentCreationWorkflow == nullptr;
+	return !HasActiveWorkflow();
 }
 
 bool FSandboxControlsViewModel::CanLoadSandbox(const FString& InSandboxRoot, FText* OutReason) const
@@ -50,14 +63,38 @@ bool FSandboxControlsViewModel::CanDeleteSandbox(const FString& InSandboxRoot, F
 	return Model->CanDeleteSandbox(InSandboxRoot, OutReason);
 }
 
-void FSandboxControlsViewModel::DeleteSandbox(const FString& InSandboxRoot) const
+void FSandboxControlsViewModel::DeleteSandbox(const FString& InSandboxRoot)
 {
+	CancelWorkflows();
 	Model->DeleteSandbox(InSandboxRoot);
 }
 
-void FSandboxControlsViewModel::OnWorkflowEnded()
+bool FSandboxControlsViewModel::CanRenameSandbox(const FString& InSandboxRoot, FText* OutError) const
+{
+	return !HasActiveWorkflow() && Model->IsAllowedToRenameSandbox(InSandboxRoot, OutError);
+}
+
+void FSandboxControlsViewModel::StartRenameWorkflow(const FString& InSandboxRoot)
+{
+	if (CanRenameSandbox(InSandboxRoot))
+	{
+		CurrentRenameWorkflow = MakeUnique<FRenameSandboxWorkflow>(
+			InSandboxRoot, Model, 
+			FSimpleDelegate::CreateRaw(this, &FSandboxControlsViewModel::HandleRenameWorkflowEnded)
+			);
+		OnRenameWorkflowStartedDelegate.Broadcast(*CurrentRenameWorkflow);
+	}
+}
+
+void FSandboxControlsViewModel::HandleCreationWorkflowEnded()
 {
 	CurrentCreationWorkflow.Reset();
-	OnCreationWorkflowEnded().Broadcast();
+	OnCreationWorkflowEndedDelegate.Broadcast();
+}
+
+void FSandboxControlsViewModel::HandleRenameWorkflowEnded()
+{
+	CurrentRenameWorkflow.Reset();
+	OnRenameWorkflowEndedDelegate.Broadcast();
 }
 }
diff --git a/Engine/Plugins/Developer/Sandbox/SandboxedEditing/Source/SandboxedEditing/Private/Features/Browser/ViewModels/SandboxControlsViewModel.h b/Engine/Plugins/Developer/Sandbox/SandboxedEditing/Source/SandboxedEditing/Private/Features/Browser/ViewModels/SandboxControlsViewModel.h
--- a/Engine/Plugins/Developer/Sandbox/SandboxedEditing/Source/SandboxedEditing/Private/Features/Browser/ViewModels/SandboxControlsViewModel.h
+++ b/Engine/Plugins/Developer/Sandbox/SandboxedEditing/Source/SandboxedEditing/Private/Features/Browser/ViewModels/SandboxControlsViewModel.h
@@ -4,13 +4,15 @@
 
 #include "Delegates/Delegate.h"
 #include "Framework/Models/SandboxSystemModel.h"
+#include "RenameSandboxWorkflow.h"
 #include "SandboxCreationWorkflow.h"
 #include "Templates/UniquePtr.h"
 #include "Templates/UnrealTemplate.h"
 
 namespace UE::SandboxedEditing
 {
-DECLARE_MULTICAST_DELEGATE_OneParam(FCreationWorkflowDelegate, const FSandboxCreationWorkflow&);
+DECLARE_MULTICAST_DELEGATE_OneParam(FCreationWorkflowDelegate, FSandboxCreationWorkflow&);
+DECLARE_MULTICAST_DELEGATE_OneParam(FRenameWorkflowDelegate, FRenameSandboxWorkflow&);
 
 /**
  * Handles actions that can be performed with FSandboxListItems, such as deleting, renaming, loading.
@@ -24,6 +26,9 @@ public:
 	
 	explicit FSandboxControlsViewModel(const TSharedRef<FSandboxSystemModel>& InModel);
 	
+	/** Cancels any workflows that may be active. */
+	void CancelWorkflows();
+	
 	/** Starts a workflow for creating a new sandbox. */
 	void StartCreationWorkflow();
 	/** @return Whether a new sandbox creation workflow can currently be started. */
@@ -39,28 +44,52 @@ public:
 	/** @return Whether the sandbox can be deleted. */
 	bool CanDeleteSandbox(const FString& InSandboxRoot, FText* OutReason = nullptr) const;
 	/** Deletes the sandbox. */
-	void DeleteSandbox(const FString& InSandboxRoot) const;
+	void DeleteSandbox(const FString& InSandboxRoot);
+	
+	/** @return Whether a rename operation can be started for the sandbox. */
+	bool CanRenameSandbox(const FString& InSandboxRoot, FText* OutError = nullptr) const;
+	/** Starts a rename operation for the sandbox. */
+	void StartRenameWorkflow(const FString& InSandboxRoot);
+	/** @return The currently active rename workflow. */
+	FRenameSandboxWorkflow* GetCurrentRenameWorkflow() const { return CurrentRenameWorkflow.Get(); }
 	
 	/** @return Delegate that is invoked when the creation workflow starts. */
 	FCreationWorkflowDelegate& OnCreationWorkflowStarted() { return OnCreationWorkflowStartedDelegate; }
 	/** @return Delegate that is invoked when the creation workflow ends. */
 	FSimpleMulticastDelegate& OnCreationWorkflowEnded() { return OnCreationWorkflowEndedDelegate; }
 	
+	/** @return Delegate that is invoked when the rename workflow starts. */
+	FRenameWorkflowDelegate& OnRenameWorkflowStarted() { return OnRenameWorkflowStartedDelegate; }
+	/** @return Delegate that is invoked when the rename workflow ends. */
+	FSimpleMulticastDelegate& OnRenameWorkflowEnded() { return OnRenameWorkflowEndedDelegate; }
+	
 private:
 	
 	/** The underlying model that performs the mutations. */
 	const TSharedRef<FSandboxSystemModel> Model;
 	
-	/** Set if the user is currently in a process of creating a sandbox. */
+	/** Set if the user is currently in the process of creating a sandbox. */
 	TUniquePtr<FSandboxCreationWorkflow> CurrentCreationWorkflow;
+	/** Set if the user is currently in the process of renaming a sandbox. */
+	TUniquePtr<FRenameSandboxWorkflow> CurrentRenameWorkflow;
 	
 	/** Invoked when the creation workflow starts. */
 	FCreationWorkflowDelegate OnCreationWorkflowStartedDelegate;
 	/** Invoked when the creation workflow ends. */
 	FSimpleMulticastDelegate OnCreationWorkflowEndedDelegate;
 	
+	/** Invoked when the creation rename starts. */
+	FRenameWorkflowDelegate OnRenameWorkflowStartedDelegate;
+	/** Invoked when the creation rename ends. */
+	FSimpleMulticastDelegate OnRenameWorkflowEndedDelegate;
+	
 	/** Called by CurrentCreationWorkflow when the workflow ends. */
-	void OnWorkflowEnded();
+	void HandleCreationWorkflowEnded();
+	/** Called by CurrentRenameWorkflow when the workflow ends. */
+	void HandleRenameWorkflowEnded();
+	
+	/** @return Whether a workflow is active */
+	bool HasActiveWorkflow() const { return CurrentCreationWorkflow || CurrentRenameWorkflow; }
 };
 }
 
diff --git a/Engine/Plugins/Developer/Sandbox/SandboxedEditing/Source/SandboxedEditing/Private/Features/Browser/Views/BrowserViews.h b/Engine/Plugins/Developer/Sandbox/SandboxedEditing/Source/SandboxedEditing/Private/Features/Browser/Views/BrowserViews.h
--- a/Engine/Plugins/Developer/Sandbox/SandboxedEditing/Source/SandboxedEditing/Private/Features/Browser/Views/BrowserViews.h
+++ b/Engine/Plugins/Developer/Sandbox/SandboxedEditing/Source/SandboxedEditing/Private/Features/Browser/Views/BrowserViews.h
@@ -5,6 +5,7 @@
 #include "AskUserAboutDirtyPackagesView.h"
 #include "AskUserToPersistView.h"
 #include "PersistSandboxView.h"
+#include "RenameWorkflowView.h"
 #include "Templates/SharedPointer.h"
 
 namespace UE::SandboxedEditing
@@ -26,12 +27,18 @@ struct FBrowserViews
 	/** Visualizes FPersistSandboxWorkflow. */
 	const TSharedRef<FPersistSandboxView> PersistSandboxView;
 	
+	/** Displays a notification when a sandbox is renamed. */
+	const TSharedRef<FRenameWorkflowView> RenameWorkflowView;
+	
 	explicit FBrowserViews(
-		const TSharedRef<FPersistSandboxViewModel>& InPersistViewModel, const TSharedRef<FLeaveSandboxViewModel>& InLeaveViewModel
+		const TSharedRef<FSandboxControlsViewModel>& InControlsViewModel, 
+		const TSharedRef<FPersistSandboxViewModel>& InPersistViewModel, 
+		const TSharedRef<FLeaveSandboxViewModel>& InLeaveViewModel
 		)
 		: AskUserAboutDirtyPackagesView(MakeShared<FAskUserAboutDirtyPackagesView>(InLeaveViewModel))
 		, AskUserToPersistView(MakeShared<FAskUserToPersistView>(InLeaveViewModel))
 		, PersistSandboxView(MakeShared<FPersistSandboxView>(InPersistViewModel))
+		, RenameWorkflowView(MakeShared<FRenameWorkflowView>(InControlsViewModel))
 	{}
 };
 }
diff --git a/Engine/Plugins/Developer/Sandbox/SandboxedEditing/Source/SandboxedEditing/Private/Features/Browser/Views/RenameWorkflowView.cpp b/Engine/Plugins/Developer/Sandbox/SandboxedEditing/Source/SandboxedEditing/Private/Features/Browser/Views/RenameWorkflowView.cpp
new file mode 100644
--- /dev/null
+++ b/Engine/Plugins/Developer/Sandbox/SandboxedEditing/Source/SandboxedEditing/Private/Features/Browser/Views/RenameWorkflowView.cpp
@@ -0,0 +1,47 @@
+﻿// Copyright Epic Games, Inc. All Rights Reserved.
+
+#include "RenameWorkflowView.h"
+
+#include "Features/Browser/ViewModels/SandboxControlsViewModel.h"
+#include "Framework/Notifications/NotificationManager.h"
+#include "Widgets/Notifications/SNotificationList.h"
+
+#define LOCTEXT_NAMESPACE "FRenameWorkflowView"
+
+namespace UE::SandboxedEditing
+{
+FRenameWorkflowView::FRenameWorkflowView(const TSharedRef<FSandboxControlsViewModel>& InControlsViewModel)
+	: ControlsViewModel(InControlsViewModel)
+{
+	ControlsViewModel->OnRenameWorkflowStarted().AddRaw(this, &FRenameWorkflowView::OnStartRenameOperation);
+}
+
+FRenameWorkflowView::~FRenameWorkflowView()
+{
+	ControlsViewModel->OnRenameWorkflowStarted().RemoveAll(this);
+	if (FRenameSandboxWorkflow* Workflow = ControlsViewModel->GetCurrentRenameWorkflow())
+	{
+		Workflow->OnSandboxRenamed().RemoveAll(this);
+	}
+}
+
+void FRenameWorkflowView::OnStartRenameOperation(FRenameSandboxWorkflow& InWorkflow) const
+{
+	InWorkflow.OnSandboxRenamed().AddRaw(this, &FRenameWorkflowView::OnSandboxRenamed);
+}
+
+void FRenameWorkflowView::OnSandboxRenamed(const FString& InOldName, const FString& InNewName) const
+{
+	FSlateNotificationManager& NotificationManager = FSlateNotificationManager::Get();
+	const FText Title = FText::Format(
+		LOCTEXT("Renamed", "Renamed sandbox \"{0}\" as \"{1}\""),
+		FText::AsCultureInvariant(InOldName), 
+		FText::AsCultureInvariant(InNewName)
+		);
+	
+	NotificationManager.AddNotification(Title)
+		->SetCompletionState(SNotificationItem::CS_Success);
+}
+}
+
+#undef LOCTEXT_NAMESPACE
\ No newline at end of file
diff --git a/Engine/Plugins/Developer/Sandbox/SandboxedEditing/Source/SandboxedEditing/Private/Features/Browser/Views/RenameWorkflowView.h b/Engine/Plugins/Developer/Sandbox/SandboxedEditing/Source/SandboxedEditing/Private/Features/Browser/Views/RenameWorkflowView.h
new file mode 100644
--- /dev/null
+++ b/Engine/Plugins/Developer/Sandbox/SandboxedEditing/Source/SandboxedEditing/Private/Features/Browser/Views/RenameWorkflowView.h
@@ -0,0 +1,31 @@
+﻿// Copyright Epic Games, Inc. All Rights Reserved.
+
+#pragma once
+
+#include "Features/Browser/ViewModels/RenameSandboxWorkflow.h"
+#include "Templates/SharedPointer.h"
+#include "Templates/UnrealTemplate.h"
+
+namespace UE::SandboxedEditing
+{
+class FSandboxControlsViewModel;
+
+/** Displays a notification when a sandbox is renamed. */
+class FRenameWorkflowView : public FNoncopyable
+{
+public:
+	
+	explicit FRenameWorkflowView(const TSharedRef<FSandboxControlsViewModel>& InControlsViewModel);
+	~FRenameWorkflowView();
+	
+private:
+	
+	const TSharedRef<FSandboxControlsViewModel> ControlsViewModel;
+
+	/** Subscribes OnSandboxRenamed. */
+	void OnStartRenameOperation(FRenameSandboxWorkflow& InWorkflow) const;
+	/** Shows a slate notification that the sandbox was renamed. */
+	void OnSandboxRenamed(const FString& InOldName, const FString& InNewName) const;
+};
+}
+
diff --git a/Engine/Plugins/Developer/Sandbox/SandboxedEditing/Source/SandboxedEditing/Private/Features/Browser/Widgets/Unsandboxed/List/Columns/NameSandboxColumn.cpp b/Engine/Plugins/Developer/Sandbox/SandboxedEditing/Source/SandboxedEditing/Private/Features/Browser/Widgets/Unsandboxed/List/Columns/NameSandboxColumn.cpp
--- a/Engine/Plugins/Developer/Sandbox/SandboxedEditing/Source/SandboxedEditing/Private/Features/Browser/Widgets/Unsandboxed/List/Columns/NameSandboxColumn.cpp
+++ b/Engine/Plugins/Developer/Sandbox/SandboxedEditing/Source/SandboxedEditing/Private/Features/Browser/Widgets/Unsandboxed/List/Columns/NameSandboxColumn.cpp
@@ -2,8 +2,10 @@
 
 #include "NameSandboxColumn.h"
 
+#include "Containers/Ticker.h"
 #include "Features/Browser/ViewModels/List/SandboxColumns.h"
 #include "Features/Browser/ViewModels/List/SandboxListItem.h"
+#include "Features/Browser/ViewModels/SandboxControlsViewModel.h"
 #include "Framework/Models/SandboxInfo.h"
 #include "SandboxedEditingStyle.h"
 #include "Widgets/Layout/SBox.h"
@@ -14,6 +16,134 @@
 
 namespace UE::SandboxedEditing
 {
+class SSandboxNameWidget : public SCompoundWidget
+{
+public:
+	
+	SLATE_BEGIN_ARGS(SSandboxNameWidget){}
+	SLATE_END_ARGS()
+
+	void Construct(
+		const FArguments& InArgs, 
+		const TSharedRef<FSandboxControlsViewModel>& InRenameViewModel,
+		const FMakeSandboxColumnWidgetArgs& InColumnsArgs
+		)
+	{
+		RenameViewModel = InRenameViewModel;
+		Item = InColumnsArgs.RowData;
+		
+		ChildSlot
+		[
+			SAssignNew(EditWidget, SInlineEditableTextBlock)
+			.Text_Lambda([this] { return FText::AsCultureInvariant(Item->SandboxInfo.Name); })
+			.ToolTipText_Lambda([this]
+			{
+				const FString& Description = Item->SandboxInfo.Description;
+				return !Description.IsEmpty()
+					? FText::AsCultureInvariant(Description) 
+					: LOCTEXT("NoDescription", "This sandbox has no user-provided description");
+			})
+			.HintText(LOCTEXT("NameHint", "Enter name"))
+			.HighlightText(InColumnsArgs.HighlightText)
+			.OnTextCommitted(this, &SSandboxNameWidget::OnNameCommitted)
+			.IsReadOnly(false)
+			.IsSelected(InColumnsArgs.IsSelectedDelegate)
+			.OnVerifyTextChanged(this, &SSandboxNameWidget::OnValidateName)
+			.OnBeginTextEdit(this, &SSandboxNameWidget::OnBeginTextEdit)
+			.OnExitEditingMode(this, &SSandboxNameWidget::OnExitTextEdit)
+		];
+		
+		RenameViewModel->OnRenameWorkflowStarted().AddSP(this, &SSandboxNameWidget::OnRenameWorkflowStarted);
+	}
+	
+private:
+	
+	/** Informs us when a rename operation is started. */
+	TSharedPtr<FSandboxControlsViewModel> RenameViewModel;
+	
+	/** The item this widget is displaying. */
+	TSharedPtr<FSandboxListItem> Item;
+	
+	/** The widget that does the editing. Used to make it appear like an editable text field. */
+	TSharedPtr<SInlineEditableTextBlock> EditWidget;
+	
+	bool bIsExplicitlyStartingEditOp = false;
+	
+	void OnRenameWorkflowStarted(FRenameSandboxWorkflow& InWorkflow)
+	{
+		const bool bPathSameAsItem = FPaths::IsSamePath(InWorkflow.GetRenamedSandboxRoot(), Item->SandboxInfo.SandboxRoot);
+		if (bPathSameAsItem
+			&& !EditWidget->IsInEditMode()
+			&& !bIsExplicitlyStartingEditOp)
+		{
+			TGuardValue Guard(bIsExplicitlyStartingEditOp, true);
+			EditWidget->EnterEditingMode();
+		}
+	}
+	
+	void OnBeginTextEdit(const FText& Text)
+	{
+		if (!bIsExplicitlyStartingEditOp)
+		{
+			TGuardValue Guard(bIsExplicitlyStartingEditOp, true);
+			RenameViewModel->StartRenameWorkflow(Item->SandboxInfo.SandboxRoot);
+		}
+	}
+	
+	void OnExitTextEdit()
+	{
+		// OnExitTextEdit is executed just before OnNameCommitted. In that case, we don't want to cancel the operation.
+		// OnExitTextEdit is also executed when editing stops for any other reason. In that case, we do want to cancel the operation.
+		ExecuteOnGameThread(TEXT("SSandboxNameWidget"), [WeakModel = RenameViewModel.ToWeakPtr()]
+		{
+			const TSharedPtr<FSandboxControlsViewModel> ModelPin = WeakModel.Pin();
+			FRenameSandboxWorkflow* Workflow = ModelPin ? ModelPin->GetCurrentRenameWorkflow() : nullptr;
+			if (Workflow)
+			{
+				Workflow->Cancel();
+			}
+		});
+	}
+	
+	void OnNameCommitted(const FText& InNewText, ETextCommit::Type InCommitType)
+	{
+		FRenameSandboxWorkflow* Workflow = RenameViewModel->GetCurrentRenameWorkflow();
+		if (!Workflow)
+		{
+			return;
+		}
+
+		Workflow->SetName(InNewText.ToString());
+		switch (InCommitType)
+		{
+		case ETextCommit::Default: [[fallthrough]];
+		case ETextCommit::OnEnter:
+		case ETextCommit::OnUserMovedFocus:
+			
+			Workflow->Confirm();
+			break;
+		case ETextCommit::OnCleared:
+			Workflow->Cancel();
+			break;
+		}
+	}
+	
+	bool OnValidateName(const FText& InNewText, FText& OutError) const
+	{
+		if (FRenameSandboxWorkflow* Workflow = RenameViewModel->GetCurrentRenameWorkflow())
+		{
+			Workflow->SetName(InNewText.ToString());
+			return Workflow->IsNameValid(&OutError);
+		}
+		
+		return false;
+	}
+};
+
+FNameSandboxColumn::FNameSandboxColumn(const TSharedRef<FSandboxControlsViewModel>& InRenameViewModel)
+	: RenameViewModel(InRenameViewModel)
+{}
+
 void FNameSandboxColumn::PopulateSearchTerms(const TSharedPtr<FSandboxListItem>& InRowData, TArray<FString>& OutSearchTerms) const
 {
 	OutSearchTerms.Add(InRowData->SandboxInfo.Name);
@@ -35,22 +165,7 @@ TSharedRef<SWidget> FNameSandboxColumn::MakeColumnWidget(const FMakeSandboxColum
 		.VAlign(VAlign_Center)
 		.Padding(FSandboxedEditingStyle::Get().GetMargin("SandboxedEditing.Browser.NameColumn.Padding"))
 		[
-			SNew(SInlineEditableTextBlock)
-			.Text_Lambda([RowData = InArgs.RowData] { return FText::AsCultureInvariant(RowData->SandboxInfo.Name); })
-			.ToolTipText_Lambda([RowData = InArgs.RowData]
-			{
-				const FString& Description = RowData->SandboxInfo.Description;
-				return !Description.IsEmpty()
-					? FText::AsCultureInvariant(Description) 
-					: LOCTEXT("NoDescription", "This sandbox has no user-provided description");
-			})
-			.HighlightText(InArgs.HighlightText)
-			//.OnTextCommitted(this, &SSessionRow::OnSessionNameCommitted) TODO DP: Implement renaming
-			.IsReadOnly(false)
-			.IsSelected(InArgs.IsSelectedDelegate)
-			//.OnVerifyTextChanged(this, &SSessionRow::OnValidatingSessionName) TODO DP: Implement renaming
-			//.Font(FontInfo) TODO DP: Change font based on whether active
-			//.ColorAndOpacity(FontColor) TODO DP: Change font based on whether active
+			SNew(SSandboxNameWidget, RenameViewModel, InArgs)
 		];
 }
 }
diff --git a/Engine/Plugins/Developer/Sandbox/SandboxedEditing/Source/SandboxedEditing/Private/Features/Browser/Widgets/Unsandboxed/List/Columns/NameSandboxColumn.h b/Engine/Plugins/Developer/Sandbox/SandboxedEditing/Source/SandboxedEditing/Private/Features/Browser/Widgets/Unsandboxed/List/Columns/NameSandboxColumn.h
--- a/Engine/Plugins/Developer/Sandbox/SandboxedEditing/Source/SandboxedEditing/Private/Features/Browser/Widgets/Unsandboxed/List/Columns/NameSandboxColumn.h
+++ b/Engine/Plugins/Developer/Sandbox/SandboxedEditing/Source/SandboxedEditing/Private/Features/Browser/Widgets/Unsandboxed/List/Columns/NameSandboxColumn.h
@@ -7,11 +7,15 @@
 
 namespace UE::SandboxedEditing
 {
+class FSandboxControlsViewModel;
+
 /** Displays the sandbox name. */
 class FNameSandboxColumn : public ISandboxColumnBehavior, public ISandboxColumnWidgetFactory
 {
 public:
 	
+	explicit FNameSandboxColumn(const TSharedRef<FSandboxControlsViewModel>& InRenameViewModel);
+	
 	//~ Begin ISandboxColumnBehavior Interface
 	virtual void PopulateSearchTerms(const TSharedPtr<FSandboxListItem>& InRowData, TArray<FString>& OutSearchTerms) const override;
 	//~ End ISandboxColumnBehavior Interface
@@ -20,6 +24,11 @@ public:
 	virtual SHeaderRow::FColumn::FArguments MakeColumnArguments() override;
 	virtual TSharedRef<SWidget> MakeColumnWidget(const FMakeSandboxColumnWidgetArgs& InArgs) override;
 	//~ End ISandboxWidgetFactory Interface
+	
+private:
+	
+	/** View model used to rename sandboxes. */
+	const TSharedRef<FSandboxControlsViewModel> RenameViewModel;
 };
 }
 
diff --git a/Engine/Plugins/Developer/Sandbox/SandboxedEditing/Source/SandboxedEditing/Private/Features/Browser/Widgets/Unsandboxed/List/SSandboxListView.cpp b/Engine/Plugins/Developer/Sandbox/SandboxedEditing/Source/SandboxedEditing/Private/Features/Browser/Widgets/Unsandboxed/List/SSandboxListView.cpp
--- a/Engine/Plugins/Developer/Sandbox/SandboxedEditing/Source/SandboxedEditing/Private/Features/Browser/Widgets/Unsandboxed/List/SSandboxListView.cpp
+++ b/Engine/Plugins/Developer/Sandbox/SandboxedEditing/Source/SandboxedEditing/Private/Features/Browser/Widgets/Unsandboxed/List/SSandboxListView.cpp
@@ -9,6 +9,9 @@
 #include "Features/Browser/Widgets/Unsandboxed/List/SCreateSandboxListRow.h"
 #include "SSandboxListRow.h"
 #include "SandboxedEditingStyle.h"
+#include "Features/Browser/Commands/BrowserCommands.h"
+#include "Framework/Commands/GenericCommands.h"
+#include "Framework/MultiBox/MultiBoxBuilder.h"
 #include "Utils/BreakBehavior.h"
 #include "Widgets/Views/SListView.h"
 
@@ -22,21 +25,24 @@ void SSandboxListView::Construct(const FArguments& InArgs, const FViewModels& In
 	ControlsViewModel = InViewModels.ControlsViewModel;
 	HighlightTextAttr = InArgs._HighlightText;
 	ColumnFactoryMap = InArgs._ColumnFactories;
+	ContextMenuCommandList = InArgs._ContextMenuCommandList;
 	OnSelectedSandboxChangedDelegate = InArgs._OnSelectedSandboxChanged;
 	
 	ChildSlot
 	[
 		SAssignNew(ListView, SListView<TSharedPtr<FSandboxListItem>>)
 		.ListItemsSource(&Items)
+		.SelectionMode(ESelectionMode::Multi)
 		.OnGenerateRow(this, &SSandboxListView::OnGenerateRow)
 		.OnMouseButtonDoubleClick(this, &SSandboxListView::OnItemDoubleClicked)
 		.HeaderRow(MakeHeaderRow())
 		.OnSelectionChanged(this, &SSandboxListView::OnSelectionChanged)
+		.OnContextMenuOpening(this, &SSandboxListView::OnContextMenuOpening)
 	];
 	
 	ControlsViewModel->OnCreationWorkflowStarted().AddSP(this, &SSandboxListView::OnCreationWorkflowStarted);
 	ControlsViewModel->OnCreationWorkflowEnded().AddSP(this, &SSandboxListView::OnCreationWorkflowEnded);
-	if (const FSandboxCreationWorkflow* CurrentWorkflow = ControlsViewModel->GetCurrentCreationWorkflow())
+	if (FSandboxCreationWorkflow* CurrentWorkflow = ControlsViewModel->GetCurrentCreationWorkflow())
 	{
 		OnCreationWorkflowStarted(*CurrentWorkflow);
 	}
@@ -109,6 +115,7 @@ TSharedRef<ITableRow> SSandboxListView::OnGenerateRow(TSharedPtr<FSandboxListIte
 	return SNew(SSandboxListRow, InItem.ToSharedRef(), InOwnerTable)
 		.ColumnFactories(ColumnFactoryMap)
 		.Padding(Padding)
+		.IsSelected_Lambda([this, InItem]{ return ListView->IsItemSelected(InItem); })
 		.HighlightText(HighlightTextAttr);
 }
 
@@ -122,6 +129,21 @@ void SSandboxListView::OnItemDoubleClicked(TSharedPtr<FSandboxListItem> InItem)
 	ControlsViewModel->LoadSandbox(InItem->SandboxInfo.SandboxRoot);
 }
 
+TSharedPtr<SWidget> SSandboxListView::OnContextMenuOpening()
+{
+	FMenuBuilder MenuBuilder(true, ContextMenuCommandList);
+	
+	MenuBuilder.AddMenuEntry(FGenericCommands::Get().Delete);
+	
+	TArray<TSharedPtr<FSandboxListItem>> Selection = ListView->GetSelectedItems();
+	if (Selection.Num() == 1)
+	{
+		MenuBuilder.AddMenuEntry(FGenericCommands::Get().Rename);
+	}
+	
+	return MenuBuilder.MakeWidget();
+}
+
 void SSandboxListView::OnItemsChanged()
 {
 	const TArray<TSharedPtr<FSandboxListItem>>& NewItems = ListViewModel->GetItems();
@@ -137,7 +159,7 @@ void SSandboxListView::OnItemsChanged()
 	ListView->RequestListRefresh();
 }
 
-void SSandboxListView::OnCreationWorkflowStarted(const FSandboxCreationWorkflow& InWorkflow)
+void SSandboxListView::OnCreationWorkflowStarted(FSandboxCreationWorkflow& InWorkflow)
 {
 	if (ensure(!DummyItem))
 	{
diff --git a/Engine/Plugins/Developer/Sandbox/SandboxedEditing/Source/SandboxedEditing/Private/Features/Browser/Widgets/Unsandboxed/List/SSandboxListView.h b/Engine/Plugins/Developer/Sandbox/SandboxedEditing/Source/SandboxedEditing/Private/Features/Browser/Widgets/Unsandboxed/List/SSandboxListView.h
--- a/Engine/Plugins/Developer/Sandbox/SandboxedEditing/Source/SandboxedEditing/Private/Features/Browser/Widgets/Unsandboxed/List/SSandboxListView.h
+++ b/Engine/Plugins/Developer/Sandbox/SandboxedEditing/Source/SandboxedEditing/Private/Features/Browser/Widgets/Unsandboxed/List/SSandboxListView.h
@@ -8,8 +8,9 @@
 #include "Widgets/DeclarativeSyntaxSupport.h"
 #include "Widgets/SCompoundWidget.h"
 
-class SHeaderRow;
+class FUICommandList;
 class ITableRow;
+class SHeaderRow;
 class STableViewBase;
 template <typename ItemType> class SListView;
 
@@ -36,8 +37,13 @@ public:
 	SLATE_BEGIN_ARGS(SSandboxListView){}
 		/** The text to highlight in the columns. */
 		SLATE_ATTRIBUTE(FText, HighlightText)
+		
 		/** Used to build the columns. */
 		SLATE_ARGUMENT(FSandboxColumnFactoryMap, ColumnFactories)
+
+		/** The command list to use in the context menu */
+		SLATE_ARGUMENT(TSharedPtr<FUICommandList>, ContextMenuCommandList)
+		
 		/** Invoked when the selected sandbox changes. */
 		SLATE_EVENT(FSandboxDelegate, OnSelectedSandboxChanged)
 	SLATE_END_ARGS()
@@ -64,6 +70,8 @@ private:
 	FSandboxColumnFactoryMap ColumnFactoryMap;
 	/** The text to highlight in the columns. */
 	TAttribute<FText> HighlightTextAttr;
+	/** The command list to use in the context menu */
+	TSharedPtr<FUICommandList> ContextMenuCommandList;
 	
 	/** Displays the items. */
 	TSharedPtr<SListView<TSharedPtr<FSandboxListItem>>> ListView;
@@ -86,11 +94,13 @@ private:
 	TSharedRef<ITableRow> OnGenerateRow(TSharedPtr<FSandboxListItem> InItem, const TSharedRef<STableViewBase>& InOwnerTable);
 	/** Loads the clicked sandbox */
 	void OnItemDoubleClicked(TSharedPtr<FSandboxListItem> InItem);
+	/** @return Context menu to show */
+	TSharedPtr<SWidget> OnContextMenuOpening();
 
 	/** Invoked when the list needs to be regenerated. */
 	void OnItemsChanged();
 	
-	void OnCreationWorkflowStarted(const FSandboxCreationWorkflow& InWorkflow);
+	void OnCreationWorkflowStarted(FSandboxCreationWorkflow& InWorkflow);
 	void OnCreationWorkflowEnded();
 	
 	void OnSelectionChanged(TSharedPtr<FSandboxListItem> InItem, ESelectInfo::Type) const;
diff --git a/Engine/Plugins/Developer/Sandbox/SandboxedEditing/Source/SandboxedEditing/Private/Features/Browser/Widgets/Unsandboxed/SSandboxedEditingBrowser.cpp b/Engine/Plugins/Developer/Sandbox/SandboxedEditing/Source/SandboxedEditing/Private/Features/Browser/Widgets/Unsandboxed/SSandboxedEditingBrowser.cpp
--- a/Engine/Plugins/Developer/Sandbox/SandboxedEditing/Source/SandboxedEditing/Private/Features/Browser/Widgets/Unsandboxed/SSandboxedEditingBrowser.cpp
+++ b/Engine/Plugins/Developer/Sandbox/SandboxedEditing/Source/SandboxedEditing/Private/Features/Browser/Widgets/Unsandboxed/SSandboxedEditingBrowser.cpp
@@ -29,6 +29,7 @@ void SSandboxedEditingBrowser::Construct(const FArguments& InArgs, const FViewMo
 	const TSharedRef<SSandboxListView> ListViewRef = SNew(SSandboxListView, ListViewModels)
 		.ColumnFactories(InArgs._SandboxColumnFactories)
 		.HighlightText(InViewModels.FilterSandboxViewModel, &FFilterSandboxViewModel::GetSearchText)
+		.ContextMenuCommandList(CommandList)
 		.OnSelectedSandboxChanged(this, &SSandboxedEditingBrowser::OnSelectedSandboxChanged);
 	ListView = ListViewRef;
 	
@@ -63,6 +64,11 @@ void SSandboxedEditingBrowser::BindCommands()
 		FExecuteAction::CreateSP(this, &SSandboxedEditingBrowser::HandleDeleteSelection), 
 		FCanExecuteAction::CreateSP(this, &SSandboxedEditingBrowser::CanDeleteSelection)
 		);
+	
+	CommandList->MapAction(
+		GenericCommands.Rename,
+		FExecuteAction::CreateSP(this, &SSandboxedEditingBrowser::HandleRenameSelection)
+		);
 }
 
 void SSandboxedEditingBrowser::HandleDeleteSelection() const
@@ -104,6 +110,15 @@ bool SSandboxedEditingBrowser::CanDeleteSelection() const
 	return bCanDelete;
 }
 
+void SSandboxedEditingBrowser::HandleRenameSelection() const
+{
+	if (const TOptional<FString> SandboxRoot = ListView->GetSelectedItem()
+		; SandboxRoot && ControlsViewModel->CanRenameSandbox(*SandboxRoot))
+	{
+		ControlsViewModel->StartRenameWorkflow(*SandboxRoot);
+	}
+}
+
 void SSandboxedEditingBrowser::OnSelectedSandboxChanged(TOptional<FString> InSandbox) const
 {
 	if (InSandbox)
diff --git a/Engine/Plugins/Developer/Sandbox/SandboxedEditing/Source/SandboxedEditing/Private/Features/Browser/Widgets/Unsandboxed/SSandboxedEditingBrowser.h b/Engine/Plugins/Developer/Sandbox/SandboxedEditing/Source/SandboxedEditing/Private/Features/Browser/Widgets/Unsandboxed/SSandboxedEditingBrowser.h
--- a/Engine/Plugins/Developer/Sandbox/SandboxedEditing/Source/SandboxedEditing/Private/Features/Browser/Widgets/Unsandboxed/SSandboxedEditingBrowser.h
+++ b/Engine/Plugins/Developer/Sandbox/SandboxedEditing/Source/SandboxedEditing/Private/Features/Browser/Widgets/Unsandboxed/SSandboxedEditingBrowser.h
@@ -66,6 +66,9 @@ private:
 	/** @return Whether the current selection can be deleted. */
 	bool CanDeleteSelection() const;
 	
+	/** Renames the currently selected. Only does anything if exactly one item is selected. No-op, otherwise. */
+	void HandleRenameSelection() const;
+	
 	/** Notify the UnloadedSandboxFileStateViewModel when the selected sandbox changes. */
 	void OnSelectedSandboxChanged(TOptional<FString> InSandbox) const;
 };
diff --git a/Engine/Plugins/Developer/Sandbox/SandboxedEditing/Source/SandboxedEditing/Private/Framework/Models/SandboxSystemModel.cpp b/Engine/Plugins/Developer/Sandbox/SandboxedEditing/Source/SandboxedEditing/Private/Framework/Models/SandboxSystemModel.cpp
--- a/Engine/Plugins/Developer/Sandbox/SandboxedEditing/Source/SandboxedEditing/Private/Framework/Models/SandboxSystemModel.cpp
+++ b/Engine/Plugins/Developer/Sandbox/SandboxedEditing/Source/SandboxedEditing/Private/Framework/Models/SandboxSystemModel.cpp
@@ -387,6 +387,60 @@ void FSandboxSystemModel::SetDescription(const FString& InSandboxRoot, const FSt
 	});
 }
 
+bool FSandboxSystemModel::RenameSandbox(const FString& InSandboxRoot, const FString& InNewName) const
+{
+	const FString OldSandboxName = FPaths::GetBaseFilename(InSandboxRoot);
+	const FString BaseDirectory = FPaths::GetPath(InSandboxRoot);
+	return FileSandboxCore::RenameSandboxWithDirectory(InNewName, OldSandboxName, BaseDirectory);
+}
+
+#define SET_ERROR(Text) if (OutError){ *OutError = Text; }
+bool FSandboxSystemModel::CanRenameSandbox(const FString& InSandboxRoot, const FString& InNewName, FText* OutError) const
+{
+	using namespace FileSandboxCore;
+	const FString OldSandboxName = FPaths::GetBaseFilename(InSandboxRoot);
+	const FString BaseDirectory = FPaths::GetPath(InSandboxRoot);
+	
+	switch (CanRenameSandboxWithDirectory(InNewName, OldSandboxName, BaseDirectory))
+	{
+	case ESandboxRenameSuitability::Allowed: return true;
+		
+	case ESandboxRenameSuitability::SandboxDoesNotExist:
+		ensure(false); // UI must have gotten out of sync. Investigate.
+		SET_ERROR(LOCTEXT("SandboxDoesNotExist", "Sandbox does not exist anymore.")); 
+		break;
+	case ESandboxRenameSuitability::ActiveSandbox:
+		SET_ERROR(LOCTEXT("ActiveSandbox.SpecifiedName", "You cannot rename the active sandbox.")); 
+		break;
+	case ESandboxRenameSuitability::NameIsSame:
+		SET_ERROR(LOCTEXT("ActiveSandbox.NameIsSame", "The new name must be different from the current name.")); 
+		break;
+	case ESandboxRenameSuitability::InvalidDirectory:
+		SET_ERROR(LOCTEXT("InvalidDirectory", "The name contains characters that aren’t allowed in directory names.")); 
+		break;
+	case ESandboxRenameSuitability::EmptyName:
+		SET_ERROR(LOCTEXT("EmptyName", "Name cannot be empty.")); 
+		break;
+	default:
+		checkNoEntry();
+	}
+		
+	return false;
+}
+#undef SET_ERROR
+
+#define SET_ERROR(Text) if (OutError){ *OutError = Text; }
+bool FSandboxSystemModel::IsAllowedToRenameSandbox(const FString& InSandboxRoot, FText* OutError) const
+{
+	const bool bCanRename = FileSandboxCore::IsAllowedToRenameSandbox(InSandboxRoot);
+	if (!bCanRename)
+	{
+		SET_ERROR(LOCTEXT("ActiveSandbox.UnspecifiedName", "Active sandbox cannot be renamed."))
+	}
+	return bCanRename;
+}
+#undef SET_ERROR
+
 void FSandboxSystemModel::HandleLoadSandbox(FileSandboxCore::ISandboxInstance& InInstance)
 {
 	OnLoadSandboxDelegate.Broadcast();
diff --git a/Engine/Plugins/Developer/Sandbox/SandboxedEditing/Source/SandboxedEditing/Private/Framework/Models/SandboxSystemModel.h b/Engine/Plugins/Developer/Sandbox/SandboxedEditing/Source/SandboxedEditing/Private/Framework/Models/SandboxSystemModel.h
--- a/Engine/Plugins/Developer/Sandbox/SandboxedEditing/Source/SandboxedEditing/Private/Framework/Models/SandboxSystemModel.h
+++ b/Engine/Plugins/Developer/Sandbox/SandboxedEditing/Source/SandboxedEditing/Private/Framework/Models/SandboxSystemModel.h
@@ -108,6 +108,13 @@ public:
 	/** Sets the sandbox description. */
 	void SetDescription(const FString& InSandboxRoot, const FString& InDescription) const;
 	
+	/** Tries to rename InSandboxRoot to InNewName. */
+	bool RenameSandbox(const FString& InSandboxRoot, const FString& InNewName) const;
+	/** @returnWhether the sandbox at InSandboxRoot can be renamed to InNewName. */
+	bool CanRenameSandbox(const FString& InSandboxRoot, const FString& InNewName, FText* OutError = nullptr) const;
+	/** @return Whether it is allowd to rename the InSandboxRoot, i.e. it is not active. */
+	bool IsAllowedToRenameSandbox(const FString& InSandboxRoot, FText* OutError = nullptr) const;
+	
 	/** Invoked when the result of GetKnownSandboxes may have changed. */
 	FSimpleMulticastDelegate& OnKnownSandboxesChanged() { return OnKnownSandboxesChangedDelegate; }
 	
