Sandboxed Editing: Add search box for file state.
- All locations with the file state list view now have a search box, i.e. browser, active sandbox, and persist dialogue
- The content of path column has been shortened: it's now relative to the content dir, then relative the plugin dir, and finally abs path (e.g. if it's engine content being edited)

#Synced-CL 49159615

diff --git a/Engine/Plugins/Developer/Sandbox/SandboxedEditing/Source/SandboxedEditing/Private/Features/Browser/ViewModels/BrowserViewModels.h b/Engine/Plugins/Developer/Sandbox/SandboxedEditing/Source/SandboxedEditing/Private/Features/Browser/ViewModels/BrowserViewModels.h
--- a/Engine/Plugins/Developer/Sandbox/SandboxedEditing/Source/SandboxedEditing/Private/Features/Browser/ViewModels/BrowserViewModels.h
+++ b/Engine/Plugins/Developer/Sandbox/SandboxedEditing/Source/SandboxedEditing/Private/Features/Browser/ViewModels/BrowserViewModels.h
@@ -14,6 +14,7 @@
 #include "Persist/PersistSandboxViewModel.h"
 #include "SandboxControlsViewModel.h"
 #include "SandboxMetaDataViewModel.h"
+#include "FileState/FilterFileStateViewModel.h"
 #include "Templates/SharedPointer.h"
 
 namespace UE::FileSandboxUI { class IExternalSandboxActiveViewModel; }
@@ -30,11 +31,18 @@ struct FBrowserViewModels
 	/** Knows about the items that are displayed in the browser. */
 	const TSharedRef<FSandboxListViewModel> ListViewModel;
 	
+	
+	/** The filters for UnloadedSandboxFileStateViewModel */
+	const TSharedRef<FFilterFileStateViewModel> FilterUnloadedSandboxFileStateViewModel;
 	/** Knows the file action of the currently selected sandbox. */
 	const TSharedRef<FUnloadedSandboxFileStateViewModel> UnloadedSandboxFileStateViewModel;
+	
+	/** The filters for ActiveSandboxFileStateViewModel */
+	const TSharedRef<FFilterFileStateViewModel> FilterActiveSandboxFileStateViewModel;
 	/** Knows the file actions of the active sandbox. */
 	const TSharedRef<FActiveSandboxFileStateViewModel> ActiveSandboxFileStateViewModel;
 	
+	
 	/** Handles actions that the user performs on the sandboxes via the UI. */
 	const TSharedRef<FSandboxControlsViewModel> ControlsViewModel;
 	/** Handles changing metadata of a sandbox. */
@@ -46,6 +54,7 @@ struct FBrowserViewModels
 	/** Handles leaving an active sandbox. */
 	const TSharedRef<FLeaveSandboxViewModel> LeaveViewModel;
 	
+	
 	/** Used for widgets to switch states when the active sandbox changes. */
 	const TSharedRef<FActiveSandboxTrackerViewModel> ActiveSandboxTrackerViewModel;
 	/** Used for widgets that display info about the currently active sandbox. */
@@ -53,6 +62,7 @@ struct FBrowserViewModels
 	/** Knows about when an extern sandbox is being displayed. */
 	const TSharedRef<FileSandboxUI::IExternalSandboxActiveViewModel> ExternalSandboxViewModel;
 	
+	
 	explicit FBrowserViewModels(
 		const FSandboxColumnRegistry& InSandboxColumns, 
 		const FFileStateColumnRegistry& InBrowserFileActionColumns,
@@ -63,8 +73,10 @@ struct FBrowserViewModels
 		)
 		: FilterSandboxViewModel(MakeShared<FFilterSandboxViewModel>(InSandboxColumns.GetBehaviorArray()))
 		, ListViewModel(MakeShared<FSandboxListViewModel>(InModel, FilterSandboxViewModel, InSandboxColumns.ColumnBehaviors))
-		, UnloadedSandboxFileStateViewModel(MakeShared<FUnloadedSandboxFileStateViewModel>(InModel, InBrowserFileActionColumns.ColumnBehaviors))
-		, ActiveSandboxFileStateViewModel(MakeShared<FActiveSandboxFileStateViewModel>(InModel, InActiveFileActionColumns.ColumnBehaviors))
+		, FilterUnloadedSandboxFileStateViewModel(MakeShared<FFilterFileStateViewModel>(InBrowserFileActionColumns.ToBehaviorArray()))
+		, UnloadedSandboxFileStateViewModel(MakeShared<FUnloadedSandboxFileStateViewModel>(InModel, FilterUnloadedSandboxFileStateViewModel, InBrowserFileActionColumns.ColumnBehaviors))
+		, FilterActiveSandboxFileStateViewModel(MakeShared<FFilterFileStateViewModel>(InActiveFileActionColumns.ToBehaviorArray()))
+		, ActiveSandboxFileStateViewModel(MakeShared<FActiveSandboxFileStateViewModel>(InModel, FilterActiveSandboxFileStateViewModel, InActiveFileActionColumns.ColumnBehaviors))
 		, ControlsViewModel(InControlsViewModel)
 		, MetaDataViewModel(MakeShared<FSandboxMetaDataViewModel>(InModel))
 		, PersistViewModel(MakeShared<FPersistSandboxViewModel>(InModel))
diff --git a/Engine/Plugins/Developer/Sandbox/SandboxedEditing/Source/SandboxedEditing/Private/Features/Browser/ViewModels/FileState/FileStateColumnRegistry.h b/Engine/Plugins/Developer/Sandbox/SandboxedEditing/Source/SandboxedEditing/Private/Features/Browser/ViewModels/FileState/FileStateColumnRegistry.h
--- a/Engine/Plugins/Developer/Sandbox/SandboxedEditing/Source/SandboxedEditing/Private/Features/Browser/ViewModels/FileState/FileStateColumnRegistry.h
+++ b/Engine/Plugins/Developer/Sandbox/SandboxedEditing/Source/SandboxedEditing/Private/Features/Browser/ViewModels/FileState/FileStateColumnRegistry.h
@@ -18,6 +18,13 @@ struct FFileStateColumnRegistry
 	
 	/** Creates the widgets for the columns. */
 	TMap<FName, TSharedRef<IFileStateColumnWidgetFactory>> ColumnFactories;
+	
+	TArray<TSharedRef<IFileStateColumnBehavior>> ToBehaviorArray() const
+	{
+		TArray<TSharedRef<IFileStateColumnBehavior>> Behaviors;
+		ColumnBehaviors.GenerateValueArray(Behaviors);
+		return Behaviors;
+	}
 };
 	
 /** @return Columns to show in the browser view. */
diff --git a/Engine/Plugins/Developer/Sandbox/SandboxedEditing/Source/SandboxedEditing/Private/Features/Browser/ViewModels/FileState/FilterFileStateViewModel.cpp b/Engine/Plugins/Developer/Sandbox/SandboxedEditing/Source/SandboxedEditing/Private/Features/Browser/ViewModels/FileState/FilterFileStateViewModel.cpp
new file mode 100644
--- /dev/null
+++ b/Engine/Plugins/Developer/Sandbox/SandboxedEditing/Source/SandboxedEditing/Private/Features/Browser/ViewModels/FileState/FilterFileStateViewModel.cpp
@@ -0,0 +1,41 @@
+﻿// Copyright Epic Games, Inc. All Rights Reserved.
+
+#include "FilterFileStateViewModel.h"
+
+#include "IFileStateColumnBehavior.h"
+
+namespace UE::SandboxedEditing
+{
+FFilterFileStateViewModel::FFilterFileStateViewModel(TAttribute<TArray<TSharedRef<IFileStateColumnBehavior>>> InVisibleColumns)
+	: VisibleColumns(MoveTemp(InVisibleColumns))
+	, TextFilter(FFileActionTextFilter::FItemToStringArray::CreateRaw(this, &FFilterFileStateViewModel::ItemToString))
+{
+	TextFilter.OnChanged().AddLambda([this]
+	{
+		OnFilterChangedDelegate.Broadcast();
+	});
+}
+
+FText FFilterFileStateViewModel::GetSearchText() const
+{
+	return TextFilter.GetRawFilterText();
+}
+
+void FFilterFileStateViewModel::SetSearchText(const FText& InSearchText)
+{
+	TextFilter.SetRawFilterText(InSearchText);
+}
+
+bool FFilterFileStateViewModel::PassesFilter(const TSharedRef<FFileStateItem>& InItem) const
+{
+	return TextFilter.PassesFilter(InItem);
+}
+
+void FFilterFileStateViewModel::ItemToString(const TSharedRef<FFileStateItem>& InItem, TArray<FString>& OutSearchTerms)
+{
+	for (const TSharedRef<IFileStateColumnBehavior>& Column : VisibleColumns.Get())
+	{
+		Column->PopulateSearchTerms(InItem, OutSearchTerms);
+	}
+}
+}
diff --git a/Engine/Plugins/Developer/Sandbox/SandboxedEditing/Source/SandboxedEditing/Private/Features/Browser/ViewModels/FileState/FilterFileStateViewModel.h b/Engine/Plugins/Developer/Sandbox/SandboxedEditing/Source/SandboxedEditing/Private/Features/Browser/ViewModels/FileState/FilterFileStateViewModel.h
new file mode 100644
--- /dev/null
+++ b/Engine/Plugins/Developer/Sandbox/SandboxedEditing/Source/SandboxedEditing/Private/Features/Browser/ViewModels/FileState/FilterFileStateViewModel.h
@@ -0,0 +1,49 @@
+﻿// Copyright Epic Games, Inc. All Rights Reserved.
+
+#pragma once
+
+#include "Delegates/Delegate.h"
+#include "Misc/Attribute.h"
+#include "Misc/TextFilter.h"
+#include "Templates/UnrealTemplate.h"
+
+namespace UE::SandboxedEditing
+{
+class IFileStateColumnBehavior;
+struct FFileStateItem;
+
+/** Knows of all filters for filtering file state items. */
+class FFilterFileStateViewModel : public FNoncopyable
+{
+	using FFileActionTextFilter = TTextFilter<const TSharedRef<FFileStateItem>&>;
+public:
+	
+	explicit FFilterFileStateViewModel(TAttribute<TArray<TSharedRef<IFileStateColumnBehavior>>> InVisibleColumns);
+	
+	/** @return Current search text */
+	FText GetSearchText() const;
+	/** Sets the search text. Will trigger a refilter. */
+	void SetSearchText(const FText& InSearchText);
+	
+	/** @return Whether the item can be displayed. */
+	bool PassesFilter(const TSharedRef<FFileStateItem>& InItem) const;
+	
+	/** @return Delegate invoked when any of the filters changes */
+	FSimpleMulticastDelegate& OnFilterChanged() { return OnFilterChangedDelegate; }
+	
+private:
+	
+	/** Visible columns. Used for filtering visible columns by text. */
+	const TAttribute<TArray<TSharedRef<IFileStateColumnBehavior>>> VisibleColumns;
+	
+	/** Does text filltering on items */
+	FFileActionTextFilter TextFilter;
+	
+	/** Invoked when any of the filters changes */
+	FSimpleMulticastDelegate OnFilterChangedDelegate;
+	
+	/** Converts an item to string. */
+	void ItemToString(const TSharedRef<FFileStateItem>& InItem, TArray<FString>& OutSearchTerms);
+};
+}
+
diff --git a/Engine/Plugins/Developer/Sandbox/SandboxedEditing/Source/SandboxedEditing/Private/Features/Browser/ViewModels/FileState/IFileStateColumnBehavior.h b/Engine/Plugins/Developer/Sandbox/SandboxedEditing/Source/SandboxedEditing/Private/Features/Browser/ViewModels/FileState/IFileStateColumnBehavior.h
--- a/Engine/Plugins/Developer/Sandbox/SandboxedEditing/Source/SandboxedEditing/Private/Features/Browser/ViewModels/FileState/IFileStateColumnBehavior.h
+++ b/Engine/Plugins/Developer/Sandbox/SandboxedEditing/Source/SandboxedEditing/Private/Features/Browser/ViewModels/FileState/IFileStateColumnBehavior.h
@@ -2,12 +2,17 @@
 
 #pragma once
 
+#include "Features/Browser/ViewModels/Columns/ISearchableColumnBehavior.h"
+#include "Templates/SharedPointerFwd.h"
+
 namespace UE::SandboxedEditing
 {
+struct FFileStateItem;
+
 /**
- * Marker interface. For now there is no behaviour. 
+ * View-model logic that is implemented by columns.
  * TODO UE-358389: Add interface for sorting
  */
-class IFileStateColumnBehavior
+class IFileStateColumnBehavior : public ISearchableColumnBehavior<TSharedPtr<FFileStateItem>>
 {};
 }
diff --git a/Engine/Plugins/Developer/Sandbox/SandboxedEditing/Source/SandboxedEditing/Private/Features/Browser/ViewModels/FileState/IFileStateColumnWidgetFactory.h b/Engine/Plugins/Developer/Sandbox/SandboxedEditing/Source/SandboxedEditing/Private/Features/Browser/ViewModels/FileState/IFileStateColumnWidgetFactory.h
--- a/Engine/Plugins/Developer/Sandbox/SandboxedEditing/Source/SandboxedEditing/Private/Features/Browser/ViewModels/FileState/IFileStateColumnWidgetFactory.h
+++ b/Engine/Plugins/Developer/Sandbox/SandboxedEditing/Source/SandboxedEditing/Private/Features/Browser/ViewModels/FileState/IFileStateColumnWidgetFactory.h
@@ -13,9 +13,13 @@ struct FMakeFileStateColumnWidgetArgs
 {
 	/** The row data for this column widget */
 	TSharedPtr<FFileStateItem> RowData;
+	
+	/** The (search) text to highlight in widgets. */
+	TAttribute<FText> HighlightText;
 
-	explicit FMakeFileStateColumnWidgetArgs(TSharedPtr<FFileStateItem> InRowData)
+	explicit FMakeFileStateColumnWidgetArgs(TSharedPtr<FFileStateItem> InRowData, TAttribute<FText> InHighlightText)
 		: RowData(MoveTemp(InRowData))
+		, HighlightText(MoveTemp(InHighlightText))
 	{}
 };
 
diff --git a/Engine/Plugins/Developer/Sandbox/SandboxedEditing/Source/SandboxedEditing/Private/Features/Browser/ViewModels/FileState/Models/ActiveSandboxFileStateViewModel.cpp b/Engine/Plugins/Developer/Sandbox/SandboxedEditing/Source/SandboxedEditing/Private/Features/Browser/ViewModels/FileState/Models/ActiveSandboxFileStateViewModel.cpp
--- a/Engine/Plugins/Developer/Sandbox/SandboxedEditing/Source/SandboxedEditing/Private/Features/Browser/ViewModels/FileState/Models/ActiveSandboxFileStateViewModel.cpp
+++ b/Engine/Plugins/Developer/Sandbox/SandboxedEditing/Source/SandboxedEditing/Private/Features/Browser/ViewModels/FileState/Models/ActiveSandboxFileStateViewModel.cpp
@@ -11,9 +11,10 @@ namespace UE::SandboxedEditing
 {
 FActiveSandboxFileStateViewModel::FActiveSandboxFileStateViewModel(
 	const TSharedRef<FSandboxSystemModel>& InModel,
+	const TSharedRef<FFilterFileStateViewModel>& InFilterViewModel,
 	TMap<FName, TSharedRef<IFileStateColumnBehavior>> InColumns
 	)
-	: FStaticFileStateViewModel(MoveTemp(InColumns), {})
+	: FStaticFileStateViewModel(MoveTemp(InColumns), {}, InFilterViewModel)
 	, Model(InModel)
 {
 	Model->OnLoadSandbox().AddRaw(this, &FActiveSandboxFileStateViewModel::OnSandboxLoaded);
@@ -35,20 +36,20 @@ FActiveSandboxFileStateViewModel::~FActiveSandboxFileStateViewModel()
 
 void FActiveSandboxFileStateViewModel::RefreshItems()
 {
-	if (!Model->HasActiveSandbox() && Items.IsEmpty())
+	if (!Model->HasActiveSandbox() && AllUnfilteredItems.IsEmpty())
 	{
 		return;
 	}
 	
-	if (!Model->HasActiveSandbox() && !Items.IsEmpty())
+	if (!Model->HasActiveSandbox() && !AllUnfilteredItems.IsEmpty())
 	{
-		OnItemsChangedDelegate.Broadcast();
-		Items.Empty();
+		AllUnfilteredItems.Empty();
+		NotifyItemsChanged();
 		return;
 	}
 	
 	bool bMadeChanges = false;
-	TArray<TSharedPtr<FFileStateItem>> OldItems = MoveTemp(Items);
+	TArray<TSharedPtr<FFileStateItem>> OldItems = MoveTemp(AllUnfilteredItems);
 	const FString ActiveSandbox = Model->GetActiveSandboxPath();
 	Model->EnumerateFileChanges(ActiveSandbox, 
 		[this, &bMadeChanges, &OldItems](const FString& InNonSandboxPath, FileSandboxCore::ESandboxFileChange InAction)
@@ -60,21 +61,21 @@ void FActiveSandboxFileStateViewModel::RefreshItems()
 			if (ReusedItem)
 			{
 				ReusedItem->Get()->Action = InAction; // Do this to make sure the item is up to date
-				Items.Emplace(*ReusedItem);
+				AllUnfilteredItems.Emplace(*ReusedItem);
 			}
 			else
 			{
 				bMadeChanges = true;
-				Items.Emplace(MakeShared<FFileStateItem>(InNonSandboxPath, InAction));
+				AllUnfilteredItems.Emplace(MakeShared<FFileStateItem>(InNonSandboxPath, InAction));
 			}
 				
 			return FileSandboxCore::EBreakBehavior::Continue;
 		});
 	
-	bMadeChanges |= OldItems.Num() != Items.Num();
+	bMadeChanges |= OldItems.Num() != AllUnfilteredItems.Num();
 	if (bMadeChanges)
 	{
-		OnItemsChangedDelegate.Broadcast();
+		NotifyItemsChanged();
 	}
 }
 
diff --git a/Engine/Plugins/Developer/Sandbox/SandboxedEditing/Source/SandboxedEditing/Private/Features/Browser/ViewModels/FileState/Models/ActiveSandboxFileStateViewModel.h b/Engine/Plugins/Developer/Sandbox/SandboxedEditing/Source/SandboxedEditing/Private/Features/Browser/ViewModels/FileState/Models/ActiveSandboxFileStateViewModel.h
--- a/Engine/Plugins/Developer/Sandbox/SandboxedEditing/Source/SandboxedEditing/Private/Features/Browser/ViewModels/FileState/Models/ActiveSandboxFileStateViewModel.h
+++ b/Engine/Plugins/Developer/Sandbox/SandboxedEditing/Source/SandboxedEditing/Private/Features/Browser/ViewModels/FileState/Models/ActiveSandboxFileStateViewModel.h
@@ -15,6 +15,7 @@ public:
 	
 	explicit FActiveSandboxFileStateViewModel(
 		const TSharedRef<FSandboxSystemModel>& InModel,
+		const TSharedRef<FFilterFileStateViewModel>& InFilterViewModel,
 		TMap<FName, TSharedRef<IFileStateColumnBehavior>> InColumns
 		);
 	~FActiveSandboxFileStateViewModel();
diff --git a/Engine/Plugins/Developer/Sandbox/SandboxedEditing/Source/SandboxedEditing/Private/Features/Browser/ViewModels/FileState/Models/StaticFileStateViewModel.cpp b/Engine/Plugins/Developer/Sandbox/SandboxedEditing/Source/SandboxedEditing/Private/Features/Browser/ViewModels/FileState/Models/StaticFileStateViewModel.cpp
--- a/Engine/Plugins/Developer/Sandbox/SandboxedEditing/Source/SandboxedEditing/Private/Features/Browser/ViewModels/FileState/Models/StaticFileStateViewModel.cpp
+++ b/Engine/Plugins/Developer/Sandbox/SandboxedEditing/Source/SandboxedEditing/Private/Features/Browser/ViewModels/FileState/Models/StaticFileStateViewModel.cpp
@@ -2,15 +2,33 @@
 
 #include "StaticFileStateViewModel.h"
 
+#include "Features/Browser/ViewModels/FileState/FilterFileStateViewModel.h"
+
 namespace UE::SandboxedEditing
 {
 FStaticFileStateViewModel::FStaticFileStateViewModel(
 	TMap<FName, TSharedRef<IFileStateColumnBehavior>> InColumns,
-	TArray<TSharedPtr<FFileStateItem>> InItems
+	TArray<TSharedPtr<FFileStateItem>> InItems,
+	TSharedPtr<FFilterFileStateViewModel> InFilterViewModel
 	)
 	: Columns(MoveTemp(InColumns))
-	, Items(MoveTemp(InItems))
-{}
+	, FilterViewModel(MoveTemp(InFilterViewModel))
+	, AllUnfilteredItems(MoveTemp(InItems))
+	, FilteredItems(AllUnfilteredItems)
+{
+	if (FilterViewModel)
+	{
+		FilterViewModel->OnFilterChanged().AddRaw(this, &FStaticFileStateViewModel::NotifyItemsChanged);
+	}
+}
+
+FStaticFileStateViewModel::~FStaticFileStateViewModel()
+{
+	if (FilterViewModel)
+	{
+		FilterViewModel->OnFilterChanged().RemoveAll(this);
+	}
+}
 
 TArray<FName> FStaticFileStateViewModel::GetDisplayedColumns() const
 {
@@ -18,4 +36,29 @@ TArray<FName> FStaticFileStateViewModel::GetDisplayedColumns() const
 	Columns.GenerateKeyArray(ColumnNames);
 	return ColumnNames;
 }
+
+void FStaticFileStateViewModel::NotifyItemsChanged()
+{
+	FilterItems();
+	OnItemsChangedDelegate.Broadcast();
+}
+
+void FStaticFileStateViewModel::FilterItems()
+{
+	if (!FilterViewModel)
+	{
+		FilteredItems = AllUnfilteredItems;
+		return;
+	}
+	
+	FilteredItems.Empty(AllUnfilteredItems.Num());
+	
+	for (const TSharedPtr<FFileStateItem>& Item : AllUnfilteredItems)
+	{
+		if (ensure(Item) && FilterViewModel->PassesFilter(Item.ToSharedRef()))
+		{
+			FilteredItems.Add(Item);
+		}
+	}
+}
 }
diff --git a/Engine/Plugins/Developer/Sandbox/SandboxedEditing/Source/SandboxedEditing/Private/Features/Browser/ViewModels/FileState/Models/StaticFileStateViewModel.h b/Engine/Plugins/Developer/Sandbox/SandboxedEditing/Source/SandboxedEditing/Private/Features/Browser/ViewModels/FileState/Models/StaticFileStateViewModel.h
--- a/Engine/Plugins/Developer/Sandbox/SandboxedEditing/Source/SandboxedEditing/Private/Features/Browser/ViewModels/FileState/Models/StaticFileStateViewModel.h
+++ b/Engine/Plugins/Developer/Sandbox/SandboxedEditing/Source/SandboxedEditing/Private/Features/Browser/ViewModels/FileState/Models/StaticFileStateViewModel.h
@@ -6,6 +6,7 @@
 
 namespace UE::SandboxedEditing
 {
+class FFilterFileStateViewModel;
 class IFileStateColumnBehavior;
 class FSandboxListItem;
 
@@ -22,11 +23,13 @@ public:
 	
 	explicit FStaticFileStateViewModel(
 		TMap<FName, TSharedRef<IFileStateColumnBehavior>> InColumns,
-		TArray<TSharedPtr<FFileStateItem>> InItems = {}
+		TArray<TSharedPtr<FFileStateItem>> InItems = {},
+		TSharedPtr<FFilterFileStateViewModel> InFilterViewModel = nullptr
 		);
+	virtual ~FStaticFileStateViewModel() override;
 	
 	//~ Begin IFileStateViewModel Interface
-	virtual const TArray<TSharedPtr<FFileStateItem>>& GetItems() const override { return Items; }
+	virtual const TArray<TSharedPtr<FFileStateItem>>& GetItems() const override { return FilteredItems; }
 	virtual TArray<FName> GetDisplayedColumns() const override;
 	virtual FSimpleMulticastDelegate& OnItemsChanged() override { return OnItemsChangedDelegate; }
 	//~ End IFileStateViewModel Interface
@@ -36,14 +39,28 @@ protected:
 	/** All columns that can ever be displayed. */
 	const TMap<FName, TSharedRef<IFileStateColumnBehavior>> Columns;
 	
+	/** Optional filter model. */
+	const TSharedPtr<FFilterFileStateViewModel> FilterViewModel;
+	
 	/** The items to display. */
-	TArray<TSharedPtr<FFileStateItem>> Items;
+	TArray<TSharedPtr<FFileStateItem>> AllUnfilteredItems;
+	
+	/** Called by subclasses when they have changed AllUnfilteredItems. */
+	void NotifyItemsChanged();
+	
+private:
 	
 	/**
 	 * Dummy to fulfill the interface requirements. Can be used by subclasses. 
 	 * TODO when sorting column is implemented, this can be invoked.
 	 */
 	FSimpleMulticastDelegate OnItemsChangedDelegate;
+	
+	/** A filtered version of AllUnfilteredItems. */
+	TArray<TSharedPtr<FFileStateItem>> FilteredItems;
+	
+	/** Refilters the items. */
+	void FilterItems();
 };
 }
 
diff --git a/Engine/Plugins/Developer/Sandbox/SandboxedEditing/Source/SandboxedEditing/Private/Features/Browser/ViewModels/FileState/Models/UnloadedSandboxFileStateViewModel.cpp b/Engine/Plugins/Developer/Sandbox/SandboxedEditing/Source/SandboxedEditing/Private/Features/Browser/ViewModels/FileState/Models/UnloadedSandboxFileStateViewModel.cpp
--- a/Engine/Plugins/Developer/Sandbox/SandboxedEditing/Source/SandboxedEditing/Private/Features/Browser/ViewModels/FileState/Models/UnloadedSandboxFileStateViewModel.cpp
+++ b/Engine/Plugins/Developer/Sandbox/SandboxedEditing/Source/SandboxedEditing/Private/Features/Browser/ViewModels/FileState/Models/UnloadedSandboxFileStateViewModel.cpp
@@ -9,9 +9,10 @@ namespace UE::SandboxedEditing
 {
 FUnloadedSandboxFileStateViewModel::FUnloadedSandboxFileStateViewModel(
 	const TSharedRef<FSandboxSystemModel>& InModel,
+	const TSharedRef<FFilterFileStateViewModel>& InFilterViewModel,
 	const TMap<FName, TSharedRef<IFileStateColumnBehavior>>& InColumns
 	)
-	: FStaticFileStateViewModel(InColumns, {})
+	: FStaticFileStateViewModel(InColumns, {}, InFilterViewModel)
 	, Model(InModel)
 {
 	Model->OnLeaveSandbox().AddRaw(this, &FUnloadedSandboxFileStateViewModel::OnSandboxLeft);
@@ -26,17 +27,17 @@ void FUnloadedSandboxFileStateViewModel::SetContent(const FString& InSandboxRoot
 {
 	SelectedSandbox = InSandboxRoot;
 	
-	Items = MakeItemsFromFileChanges(
+	AllUnfilteredItems = MakeItemsFromFileChanges(
 		Model->GatherFileChanges(InSandboxRoot)
 		);
-	OnItemsChangedDelegate.Broadcast();
+	NotifyItemsChanged();
 }
 
 void FUnloadedSandboxFileStateViewModel::ClearContent()
 {
 	SelectedSandbox.Reset();
-	Items.Empty();
-	OnItemsChangedDelegate.Broadcast();
+	AllUnfilteredItems.Empty();
+	NotifyItemsChanged();
 }
 
 void FUnloadedSandboxFileStateViewModel::OnSandboxLeft()
diff --git a/Engine/Plugins/Developer/Sandbox/SandboxedEditing/Source/SandboxedEditing/Private/Features/Browser/ViewModels/FileState/Models/UnloadedSandboxFileStateViewModel.h b/Engine/Plugins/Developer/Sandbox/SandboxedEditing/Source/SandboxedEditing/Private/Features/Browser/ViewModels/FileState/Models/UnloadedSandboxFileStateViewModel.h
--- a/Engine/Plugins/Developer/Sandbox/SandboxedEditing/Source/SandboxedEditing/Private/Features/Browser/ViewModels/FileState/Models/UnloadedSandboxFileStateViewModel.h
+++ b/Engine/Plugins/Developer/Sandbox/SandboxedEditing/Source/SandboxedEditing/Private/Features/Browser/ViewModels/FileState/Models/UnloadedSandboxFileStateViewModel.h
@@ -17,6 +17,7 @@ public:
 	
 	explicit FUnloadedSandboxFileStateViewModel(
 		const TSharedRef<FSandboxSystemModel>& InModel,
+		const TSharedRef<FFilterFileStateViewModel>& InFilterViewModel,
 		const TMap<FName, TSharedRef<IFileStateColumnBehavior>>& InColumns
 		);
 	~FUnloadedSandboxFileStateViewModel();
diff --git a/Engine/Plugins/Developer/Sandbox/SandboxedEditing/Source/SandboxedEditing/Private/Features/Browser/Widgets/SSandboxedEditingRoot.cpp b/Engine/Plugins/Developer/Sandbox/SandboxedEditing/Source/SandboxedEditing/Private/Features/Browser/Widgets/SSandboxedEditingRoot.cpp
--- a/Engine/Plugins/Developer/Sandbox/SandboxedEditing/Source/SandboxedEditing/Private/Features/Browser/Widgets/SSandboxedEditingRoot.cpp
+++ b/Engine/Plugins/Developer/Sandbox/SandboxedEditing/Source/SandboxedEditing/Private/Features/Browser/Widgets/SSandboxedEditingRoot.cpp
@@ -20,7 +20,8 @@ void SSandboxedEditingRoot::Construct(const FArguments& InArgs, const FBrowserVi
 	ActiveSandboxTrackerViewModel = InViewModels.ActiveSandboxTrackerViewModel;
 	ExternalSandboxViewModel = InViewModels.ExternalSandboxViewModel;
 	ActiveSandboxViewModels.Emplace(
-		InViewModels.ActiveSandboxDetailsViewModel, InViewModels.MetaDataViewModel, InViewModels.ActiveSandboxFileStateViewModel
+		InViewModels.ActiveSandboxDetailsViewModel, InViewModels.MetaDataViewModel, 
+		InViewModels.ActiveSandboxFileStateViewModel, InViewModels.FilterActiveSandboxFileStateViewModel
 		);
 	ActiveFileStateColumnsAttr = InArgs._ActiveFileStateColumnsFactories;
 	
@@ -31,7 +32,7 @@ void SSandboxedEditingRoot::Construct(const FArguments& InArgs, const FBrowserVi
 	const SSandboxedEditingBrowser::FViewModels ViewModels 
 	{ 
 		InViewModels.FilterSandboxViewModel, InViewModels.ListViewModel, InViewModels.ControlsViewModel, InViewModels.MetaDataViewModel,
-		InViewModels.UnloadedSandboxFileStateViewModel
+		InViewModels.UnloadedSandboxFileStateViewModel, InViewModels.FilterUnloadedSandboxFileStateViewModel
 	};
 	Browser = SNew(SSandboxedEditingBrowser, ViewModels)
 		.CommandList(CommandBindings)
diff --git a/Engine/Plugins/Developer/Sandbox/SandboxedEditing/Source/SandboxedEditing/Private/Features/Browser/Widgets/Sandboxed/SActiveSandbox.cpp b/Engine/Plugins/Developer/Sandbox/SandboxedEditing/Source/SandboxedEditing/Private/Features/Browser/Widgets/Sandboxed/SActiveSandbox.cpp
--- a/Engine/Plugins/Developer/Sandbox/SandboxedEditing/Source/SandboxedEditing/Private/Features/Browser/Widgets/Sandboxed/SActiveSandbox.cpp
+++ b/Engine/Plugins/Developer/Sandbox/SandboxedEditing/Source/SandboxedEditing/Private/Features/Browser/Widgets/Sandboxed/SActiveSandbox.cpp
@@ -4,7 +4,7 @@
 
 #include "Features/Browser/ViewModels/Active/ActiveSandboxDetailsViewModel.h"
 #include "Features/Browser/ViewModels/FileState/Models/ActiveSandboxFileStateViewModel.h"
-#include "Features/Browser/Widgets/Shared/FileState/SFileStateListView.h"
+#include "Features/Browser/Widgets/Shared/FileState/SFilterableFileStateListView.h"
 #include "Features/Browser/Widgets/Shared/SSandboxDescription.h"
 #include "SActiveSandboxToolbar.h"
 #include "Styling/ToolBarStyle.h"
@@ -95,11 +95,8 @@ TSharedRef<SWidget> SActiveSandbox::MakeFileChangesContent(const FArguments& InA
 		]
 		.BodyContent()
 		[
-			SNew(SBorder)
-			[
-				SNew(SFileStateListView, InViewModels.ActiveSandboxFileStateModel)
-				.ColumnFactories(InArgs._ColumnFactories)
-			]
+			SNew(SFilterableFileStateListView, InViewModels.ActiveSandboxFileStateModel, InViewModels.FilterActiveSandboxFileStateModel)
+			.ColumnFactories(InArgs._ColumnFactories)
 		];
 }
 }
diff --git a/Engine/Plugins/Developer/Sandbox/SandboxedEditing/Source/SandboxedEditing/Private/Features/Browser/Widgets/Sandboxed/SActiveSandbox.h b/Engine/Plugins/Developer/Sandbox/SandboxedEditing/Source/SandboxedEditing/Private/Features/Browser/Widgets/Sandboxed/SActiveSandbox.h
--- a/Engine/Plugins/Developer/Sandbox/SandboxedEditing/Source/SandboxedEditing/Private/Features/Browser/Widgets/Sandboxed/SActiveSandbox.h
+++ b/Engine/Plugins/Developer/Sandbox/SandboxedEditing/Source/SandboxedEditing/Private/Features/Browser/Widgets/Sandboxed/SActiveSandbox.h
@@ -10,6 +10,7 @@ class FUICommandList;
 
 namespace UE::SandboxedEditing
 {
+class FFilterFileStateViewModel;
 class FActiveSandboxFileStateViewModel;
 class FActiveSandboxDetailsViewModel;
 class FSandboxMetaDataViewModel;
@@ -24,15 +25,18 @@ public:
 		TSharedRef<FActiveSandboxDetailsViewModel> ActiveSandboxViewModel;
 		TSharedRef<FSandboxMetaDataViewModel> MetaDataViewModel;
 		TSharedRef<FActiveSandboxFileStateViewModel> ActiveSandboxFileStateModel;
+		TSharedRef<FFilterFileStateViewModel> FilterActiveSandboxFileStateModel;
 
 		explicit FViewModels(
 			const TSharedRef<FActiveSandboxDetailsViewModel>& InActiveSandboxViewModel,
 			const TSharedRef<FSandboxMetaDataViewModel>& InMetaDataViewModel,
-			const TSharedRef<FActiveSandboxFileStateViewModel>& InActiveSandboxFileStateModel
+			const TSharedRef<FActiveSandboxFileStateViewModel>& InActiveSandboxFileStateModel,
+			const TSharedRef<FFilterFileStateViewModel>& InFilterActiveSandboxFileStateModel
 			)
 			: ActiveSandboxViewModel(InActiveSandboxViewModel)
 			, MetaDataViewModel(InMetaDataViewModel)
 			, ActiveSandboxFileStateModel(InActiveSandboxFileStateModel)
+			, FilterActiveSandboxFileStateModel(InFilterActiveSandboxFileStateModel)
 		{}
 	};
 	
diff --git a/Engine/Plugins/Developer/Sandbox/SandboxedEditing/Source/SandboxedEditing/Private/Features/Browser/Widgets/Sandboxed/SPersistOperationWidget.cpp b/Engine/Plugins/Developer/Sandbox/SandboxedEditing/Source/SandboxedEditing/Private/Features/Browser/Widgets/Sandboxed/SPersistOperationWidget.cpp
--- a/Engine/Plugins/Developer/Sandbox/SandboxedEditing/Source/SandboxedEditing/Private/Features/Browser/Widgets/Sandboxed/SPersistOperationWidget.cpp
+++ b/Engine/Plugins/Developer/Sandbox/SandboxedEditing/Source/SandboxedEditing/Private/Features/Browser/Widgets/Sandboxed/SPersistOperationWidget.cpp
@@ -4,9 +4,10 @@
 
 #include "Features/Browser/ViewModels/FileState/FileStateColumnRegistry.h"
 #include "Features/Browser/ViewModels/FileState/FileStateItemUtils.h"
+#include "Features/Browser/ViewModels/FileState/FilterFileStateViewModel.h"
 #include "Features/Browser/ViewModels/FileState/Models/StaticFileStateViewModel.h"
 #include "Features/Browser/ViewModels/Persist/PersistOperationViewModel.h"
-#include "Features/Browser/Widgets/Shared/FileState/SFileStateListView.h"
+#include "Features/Browser/Widgets/Shared/FileState/SFilterableFileStateListView.h"
 #include "SPositiveActionButton.h"
 #include "Styling/AppStyle.h"
 #include "Types/GatheredFileChanges.h"
@@ -24,9 +25,11 @@ void SPersistOperationWidget::Construct(const FArguments& InArgs, const TSharedR
 	ViewModel = InViewModel;
 	
 	const FFileStateColumnRegistry Columns = GetColumnsForPersist(ViewModel.ToSharedRef());
+	const TSharedRef<FFilterFileStateViewModel> FilterViewModel = MakeShared<FFilterFileStateViewModel>(Columns.ToBehaviorArray());
 	const TSharedRef<FStaticFileStateViewModel> FileStateViewModel = MakeShared<FStaticFileStateViewModel>(
 		Columns.ColumnBehaviors, 
-		MakeItemsFromFileChanges(ViewModel->GetPersistableFiles())
+		MakeItemsFromFileChanges(ViewModel->GetPersistableFiles()),
+		FilterViewModel
 		);
 	
 	ChildSlot
@@ -39,11 +42,8 @@ void SPersistOperationWidget::Construct(const FArguments& InArgs, const TSharedR
 			+SVerticalBox::Slot()
 			.Padding(5.f)
 			[
-				SNew(SBorder)
-				[
-					SNew(SFileStateListView, FileStateViewModel)
-					.ColumnFactories(Columns.ColumnFactories)
-				]
+				SNew(SFilterableFileStateListView, FileStateViewModel, FilterViewModel)
+				.ColumnFactories(Columns.ColumnFactories)
 			]
 
 			+SVerticalBox::Slot()
diff --git a/Engine/Plugins/Developer/Sandbox/SandboxedEditing/Source/SandboxedEditing/Private/Features/Browser/Widgets/Shared/FileState/Columns/ActionFileStateColumn.cpp b/Engine/Plugins/Developer/Sandbox/SandboxedEditing/Source/SandboxedEditing/Private/Features/Browser/Widgets/Shared/FileState/Columns/ActionFileStateColumn.cpp
--- a/Engine/Plugins/Developer/Sandbox/SandboxedEditing/Source/SandboxedEditing/Private/Features/Browser/Widgets/Shared/FileState/Columns/ActionFileStateColumn.cpp
+++ b/Engine/Plugins/Developer/Sandbox/SandboxedEditing/Source/SandboxedEditing/Private/Features/Browser/Widgets/Shared/FileState/Columns/ActionFileStateColumn.cpp
@@ -41,6 +41,25 @@ static FText GetActionToolTipText(FileSandboxCore::ESandboxFileChange InChange)
 }
 }
 
+void FActionFileStateColumn::PopulateSearchTerms(const TSharedPtr<FFileStateItem>& InRowData, TArray<FString>& OutSearchTerms) const
+{
+	switch (InRowData->Action)
+	{
+	case FileSandboxCore::ESandboxFileChange::Added:
+		OutSearchTerms.Add(TEXT("add"));
+		break;
+	case FileSandboxCore::ESandboxFileChange::Removed:
+		OutSearchTerms.Add(TEXT("remove"));
+		OutSearchTerms.Add(TEXT("delete"));
+		break;
+	case FileSandboxCore::ESandboxFileChange::Edited:
+		OutSearchTerms.Add(TEXT("edit"));
+		OutSearchTerms.Add(TEXT("modify"));
+		break;
+	default: break;
+	}
+}
+
 SHeaderRow::FColumn::FArguments FActionFileStateColumn::MakeColumnArguments()
 {
 	return SHeaderRow::FColumn::FArguments()
diff --git a/Engine/Plugins/Developer/Sandbox/SandboxedEditing/Source/SandboxedEditing/Private/Features/Browser/Widgets/Shared/FileState/Columns/ActionFileStateColumn.h b/Engine/Plugins/Developer/Sandbox/SandboxedEditing/Source/SandboxedEditing/Private/Features/Browser/Widgets/Shared/FileState/Columns/ActionFileStateColumn.h
--- a/Engine/Plugins/Developer/Sandbox/SandboxedEditing/Source/SandboxedEditing/Private/Features/Browser/Widgets/Shared/FileState/Columns/ActionFileStateColumn.h
+++ b/Engine/Plugins/Developer/Sandbox/SandboxedEditing/Source/SandboxedEditing/Private/Features/Browser/Widgets/Shared/FileState/Columns/ActionFileStateColumn.h
@@ -11,6 +11,10 @@ class FActionFileStateColumn : public IFileStateColumnBehavior, public IFileStat
 {
 public:
 	
+	//~ Begin IFileStateColumnBehavior Interface
+	virtual void PopulateSearchTerms(const TSharedPtr<FFileStateItem>& InRowData, TArray<FString>& OutSearchTerms) const override;
+	//~ End IFileStateColumnBehavior Interface
+	
 	//~ Begin IFileStateWidgetFactory Interface
 	virtual SHeaderRow::FColumn::FArguments MakeColumnArguments() override;
 	virtual TSharedRef<SWidget> MakeColumnWidget(const FMakeFileStateColumnWidgetArgs& InArgs) override;
diff --git a/Engine/Plugins/Developer/Sandbox/SandboxedEditing/Source/SandboxedEditing/Private/Features/Browser/Widgets/Shared/FileState/Columns/PathFileStateColumn.cpp b/Engine/Plugins/Developer/Sandbox/SandboxedEditing/Source/SandboxedEditing/Private/Features/Browser/Widgets/Shared/FileState/Columns/PathFileStateColumn.cpp
--- a/Engine/Plugins/Developer/Sandbox/SandboxedEditing/Source/SandboxedEditing/Private/Features/Browser/Widgets/Shared/FileState/Columns/PathFileStateColumn.cpp
+++ b/Engine/Plugins/Developer/Sandbox/SandboxedEditing/Source/SandboxedEditing/Private/Features/Browser/Widgets/Shared/FileState/Columns/PathFileStateColumn.cpp
@@ -12,6 +12,39 @@
 
 namespace UE::SandboxedEditing
 {
+namespace PathFileColumnDetail
+{
+/** @return A file path to display in the column, kept as short as possible. */
+static FString MakeProjectRelative(const FString& InPath)
+{
+	const FString AbsPath = FPaths::ConvertRelativePathToFull(InPath);
+	
+	// Most content will be in content folder, so avoid the Content/ prefix...
+	FString Result = AbsPath;
+	if (FPaths::MakePathRelativeTo(Result, *FPaths::ProjectContentDir()))
+	{
+		return Result;
+	}
+	
+	// ... if it's plugin content, we'll prefix with /Plugins/x
+	Result = AbsPath;
+	if (FPaths::MakePathRelativeTo(Result, *FPaths::ProjectDir()))
+	{
+		return Result;
+	}
+	
+	// And otherwise absolute path, e.g. if editing engine content
+	return AbsPath;
+}
+}
+
+void FPathFileStateColumn::PopulateSearchTerms(const TSharedPtr<FFileStateItem>& InRowData, TArray<FString>& OutSearchTerms) const
+{
+	OutSearchTerms.Add(
+		PathFileColumnDetail::MakeProjectRelative(InRowData->NonSandboxFile)
+		);
+}
+
 SHeaderRow::FColumn::FArguments FPathFileStateColumn::MakeColumnArguments()
 {
 	return SHeaderRow::FColumn::FArguments()
@@ -24,14 +57,21 @@ SHeaderRow::FColumn::FArguments FPathFileStateColumn::MakeColumnArguments()
 
 TSharedRef<SWidget> FPathFileStateColumn::MakeColumnWidget(const FMakeFileStateColumnWidgetArgs& InArgs)
 {
-	const FText AbsPathAsText = FText::AsCultureInvariant(FPaths::ConvertRelativePathToFull(InArgs.RowData->NonSandboxFile));
+	const FText DisplayText = FText::AsCultureInvariant(
+		PathFileColumnDetail::MakeProjectRelative(InArgs.RowData->NonSandboxFile)
+		);
+	const FText AbsPath =  FText::AsCultureInvariant(
+		FPaths::ConvertRelativePathToFull(InArgs.RowData->NonSandboxFile)
+		);
 	return SNew(SBox)
 		.VAlign(VAlign_Center)
 		.Padding(FSandboxedEditingStyle::Get().GetMargin("SandboxedEditing.FileActions.PathColumn.Padding"))
 		[
 			SNew(STextBlock)
-			.Text(AbsPathAsText)
+			.Text(DisplayText)
+			.ToolTipText(AbsPath)
 			.OverflowPolicy(ETextOverflowPolicy::MiddleEllipsis)
+			.HighlightText(InArgs.HighlightText)
 		];
 }
 }
diff --git a/Engine/Plugins/Developer/Sandbox/SandboxedEditing/Source/SandboxedEditing/Private/Features/Browser/Widgets/Shared/FileState/Columns/PathFileStateColumn.h b/Engine/Plugins/Developer/Sandbox/SandboxedEditing/Source/SandboxedEditing/Private/Features/Browser/Widgets/Shared/FileState/Columns/PathFileStateColumn.h
--- a/Engine/Plugins/Developer/Sandbox/SandboxedEditing/Source/SandboxedEditing/Private/Features/Browser/Widgets/Shared/FileState/Columns/PathFileStateColumn.h
+++ b/Engine/Plugins/Developer/Sandbox/SandboxedEditing/Source/SandboxedEditing/Private/Features/Browser/Widgets/Shared/FileState/Columns/PathFileStateColumn.h
@@ -12,6 +12,10 @@ class FPathFileStateColumn : public IFileStateColumnBehavior, public IFileStateC
 {
 public:
 	
+	//~ Begin IFileStateColumnBehavior Interface
+	virtual void PopulateSearchTerms(const TSharedPtr<FFileStateItem>& InRowData, TArray<FString>& OutSearchTerms) const override;
+	//~ End IFileStateColumnBehavior Interface
+	
 	//~ Begin IFileStateWidgetFactory Interface
 	virtual SHeaderRow::FColumn::FArguments MakeColumnArguments() override;
 	virtual TSharedRef<SWidget> MakeColumnWidget(const FMakeFileStateColumnWidgetArgs& InArgs) override;
diff --git a/Engine/Plugins/Developer/Sandbox/SandboxedEditing/Source/SandboxedEditing/Private/Features/Browser/Widgets/Shared/FileState/Columns/PersistFileStateColumn.h b/Engine/Plugins/Developer/Sandbox/SandboxedEditing/Source/SandboxedEditing/Private/Features/Browser/Widgets/Shared/FileState/Columns/PersistFileStateColumn.h
--- a/Engine/Plugins/Developer/Sandbox/SandboxedEditing/Source/SandboxedEditing/Private/Features/Browser/Widgets/Shared/FileState/Columns/PersistFileStateColumn.h
+++ b/Engine/Plugins/Developer/Sandbox/SandboxedEditing/Source/SandboxedEditing/Private/Features/Browser/Widgets/Shared/FileState/Columns/PersistFileStateColumn.h
@@ -15,6 +15,10 @@ public:
 	
 	explicit FPersistFileStateColumn(const TSharedRef<FPersistOperationViewModel>& InPersistViewModel);
 
+	//~ Begin IFileStateColumnBehavior Interface
+	virtual void PopulateSearchTerms(const TSharedPtr<FFileStateItem>& InRowData, TArray<FString>& OutSearchTerms) const override {}
+	//~ End IFileStateColumnBehavior Interface
+
 	//~ Begin IFileStateWidgetFactory Interface
 	virtual SHeaderRow::FColumn::FArguments MakeColumnArguments() override;
 	virtual TSharedRef<SWidget> MakeColumnWidget(const FMakeFileStateColumnWidgetArgs& InArgs) override;
diff --git a/Engine/Plugins/Developer/Sandbox/SandboxedEditing/Source/SandboxedEditing/Private/Features/Browser/Widgets/Shared/FileState/SFileStateFilters.cpp b/Engine/Plugins/Developer/Sandbox/SandboxedEditing/Source/SandboxedEditing/Private/Features/Browser/Widgets/Shared/FileState/SFileStateFilters.cpp
new file mode 100644
--- /dev/null
+++ b/Engine/Plugins/Developer/Sandbox/SandboxedEditing/Source/SandboxedEditing/Private/Features/Browser/Widgets/Shared/FileState/SFileStateFilters.cpp
@@ -0,0 +1,40 @@
+﻿// Copyright Epic Games, Inc. All Rights Reserved.
+
+#include "SFileStateFilters.h"
+
+#include "Features/Browser/ViewModels/FileState/FilterFileStateViewModel.h"
+#include "Widgets/Input/SSearchBox.h"
+
+#define LOCTEXT_NAMESPACE "SFileStateFilters"
+
+namespace UE::SandboxedEditing
+{
+void SFileStateFilters::Construct(const FArguments& InArgs, const TSharedRef<FFilterFileStateViewModel>& InFilterModel)
+{
+	FilterViewModel = InFilterModel;
+	
+	ChildSlot
+	[
+		SNew(SSearchBox)
+		.HintText(LOCTEXT("SearchHint", "Search file states"))
+		.OnTextChanged(this, &SFileStateFilters::OnSearchTextChanged)
+		.OnTextCommitted(this, &SFileStateFilters::OnSearchTextCommitted)
+		.DelayChangeNotificationsWhileTyping(true)
+	];
+}
+
+void SFileStateFilters::OnSearchTextChanged(const FText& Text) const
+{
+	FilterViewModel->SetSearchText(Text);
+}
+
+void SFileStateFilters::OnSearchTextCommitted(const FText& InFilterText, ETextCommit::Type Arg)
+{
+	if (!InFilterText.EqualTo(FilterViewModel->GetSearchText()))
+	{
+		OnSearchTextChanged(InFilterText);
+	}
+}
+}
+
+#undef LOCTEXT_NAMESPACE
\ No newline at end of file
diff --git a/Engine/Plugins/Developer/Sandbox/SandboxedEditing/Source/SandboxedEditing/Private/Features/Browser/Widgets/Shared/FileState/SFileStateFilters.h b/Engine/Plugins/Developer/Sandbox/SandboxedEditing/Source/SandboxedEditing/Private/Features/Browser/Widgets/Shared/FileState/SFileStateFilters.h
new file mode 100644
--- /dev/null
+++ b/Engine/Plugins/Developer/Sandbox/SandboxedEditing/Source/SandboxedEditing/Private/Features/Browser/Widgets/Shared/FileState/SFileStateFilters.h
@@ -0,0 +1,31 @@
+﻿// Copyright Epic Games, Inc. All Rights Reserved.
+
+#pragma once
+
+#include "Widgets/DeclarativeSyntaxSupport.h"
+#include "Widgets/SCompoundWidget.h"
+
+namespace UE::SandboxedEditing
+{
+class FFilterFileStateViewModel;
+
+/** Displays the widgets required to filter file actions, i.e. FFilterFileStateViewModel. */
+class SFileStateFilters : public SCompoundWidget
+{
+public:
+	
+	SLATE_BEGIN_ARGS(SFileStateFilters) {}
+	SLATE_END_ARGS()
+	
+	void Construct(const FArguments& InArgs, const TSharedRef<FFilterFileStateViewModel>& InFilterModel);
+	
+private:
+	
+	/** Knows how to filter file state items. */
+	TSharedPtr<FFilterFileStateViewModel> FilterViewModel;
+	
+	void OnSearchTextChanged(const FText& Text) const;
+	void OnSearchTextCommitted(const FText& InFilterText, ETextCommit::Type Arg);
+};
+}
+
diff --git a/Engine/Plugins/Developer/Sandbox/SandboxedEditing/Source/SandboxedEditing/Private/Features/Browser/Widgets/Shared/FileState/SFileStateListRow.cpp b/Engine/Plugins/Developer/Sandbox/SandboxedEditing/Source/SandboxedEditing/Private/Features/Browser/Widgets/Shared/FileState/SFileStateListRow.cpp
--- a/Engine/Plugins/Developer/Sandbox/SandboxedEditing/Source/SandboxedEditing/Private/Features/Browser/Widgets/Shared/FileState/SFileStateListRow.cpp
+++ b/Engine/Plugins/Developer/Sandbox/SandboxedEditing/Source/SandboxedEditing/Private/Features/Browser/Widgets/Shared/FileState/SFileStateListRow.cpp
@@ -10,6 +10,7 @@ void SFileStateListRow::Construct(const FArguments& InArgs, const TSharedRef<FFi
 {
 	Item = InItem;
 	ColumnFactoryMap = InArgs._ColumnFactories;
+	HighlightText = InArgs._HighlightText;
 	
 	SMultiColumnTableRow::Construct(
 		SMultiColumnTableRow::FArguments().Padding(InArgs._Padding),
@@ -22,7 +23,7 @@ TSharedRef<SWidget> SFileStateListRow::GenerateWidgetForColumn(const FName& InCo
 	if (TSharedRef<IFileStateColumnWidgetFactory>* Factory = ColumnFactoryMap.Find(InColumnName))
 	{
 		return Factory->Get().MakeColumnWidget(
-			FMakeFileStateColumnWidgetArgs(Item)
+			FMakeFileStateColumnWidgetArgs(Item, HighlightText)
 		);
 	}
 	
diff --git a/Engine/Plugins/Developer/Sandbox/SandboxedEditing/Source/SandboxedEditing/Private/Features/Browser/Widgets/Shared/FileState/SFileStateListRow.h b/Engine/Plugins/Developer/Sandbox/SandboxedEditing/Source/SandboxedEditing/Private/Features/Browser/Widgets/Shared/FileState/SFileStateListRow.h
--- a/Engine/Plugins/Developer/Sandbox/SandboxedEditing/Source/SandboxedEditing/Private/Features/Browser/Widgets/Shared/FileState/SFileStateListRow.h
+++ b/Engine/Plugins/Developer/Sandbox/SandboxedEditing/Source/SandboxedEditing/Private/Features/Browser/Widgets/Shared/FileState/SFileStateListRow.h
@@ -21,6 +21,9 @@ public:
 		SLATE_ARGUMENT(FFileStateColumnFactoryMap, ColumnFactories)
 		/** Padding for the row */
 		SLATE_ARGUMENT(FMargin, Padding)
+		
+		/** The text to highlight in the columns. */
+		SLATE_ATTRIBUTE(FText, HighlightText)
 	SLATE_END_ARGS()
 	
 	void Construct(const FArguments& InArgs, const TSharedRef<FFileStateItem>& InItem, const TSharedRef<STableViewBase>& InOwnerTable);
@@ -36,6 +39,9 @@ private:
 	
 	/** Used to build the columns. */
 	FFileStateColumnFactoryMap ColumnFactoryMap;
+	
+	/** The text to highlight in the columns. */
+	TAttribute<FText> HighlightText;
 };
 }
 
diff --git a/Engine/Plugins/Developer/Sandbox/SandboxedEditing/Source/SandboxedEditing/Private/Features/Browser/Widgets/Shared/FileState/SFileStateListView.cpp b/Engine/Plugins/Developer/Sandbox/SandboxedEditing/Source/SandboxedEditing/Private/Features/Browser/Widgets/Shared/FileState/SFileStateListView.cpp
--- a/Engine/Plugins/Developer/Sandbox/SandboxedEditing/Source/SandboxedEditing/Private/Features/Browser/Widgets/Shared/FileState/SFileStateListView.cpp
+++ b/Engine/Plugins/Developer/Sandbox/SandboxedEditing/Source/SandboxedEditing/Private/Features/Browser/Widgets/Shared/FileState/SFileStateListView.cpp
@@ -14,6 +14,7 @@ void SFileStateListView::Construct(const FArguments& InArgs, const TSharedRef<IF
 {
 	ViewModel = InViewModel;
 	ColumnFactories = InArgs._ColumnFactories;
+	HighlightText = InArgs._HighlightText;
 	
 	ChildSlot
 	[
@@ -49,6 +50,7 @@ TSharedRef<ITableRow> SFileStateListView::OnGenerateRow(TSharedPtr<FFileStateIte
 	
 	return SNew(SFileStateListRow, InItem.ToSharedRef(), InOwnerTable)
 		.ColumnFactories(ColumnFactories)
-		.Padding(Padding);
+		.Padding(Padding)
+		.HighlightText(HighlightText);
 }
 }
diff --git a/Engine/Plugins/Developer/Sandbox/SandboxedEditing/Source/SandboxedEditing/Private/Features/Browser/Widgets/Shared/FileState/SFileStateListView.h b/Engine/Plugins/Developer/Sandbox/SandboxedEditing/Source/SandboxedEditing/Private/Features/Browser/Widgets/Shared/FileState/SFileStateListView.h
--- a/Engine/Plugins/Developer/Sandbox/SandboxedEditing/Source/SandboxedEditing/Private/Features/Browser/Widgets/Shared/FileState/SFileStateListView.h
+++ b/Engine/Plugins/Developer/Sandbox/SandboxedEditing/Source/SandboxedEditing/Private/Features/Browser/Widgets/Shared/FileState/SFileStateListView.h
@@ -23,6 +23,8 @@ public:
 	SLATE_BEGIN_ARGS(SFileStateListView) {}
 		/** Used to build the columns. */
 		SLATE_ARGUMENT(FFileStateColumnFactoryMap, ColumnFactories)
+		/** The text to highlight in the columns. */
+		SLATE_ATTRIBUTE(FText, HighlightText)
 	SLATE_END_ARGS()
 	
 	void Construct(const FArguments& InArgs, const TSharedRef<IFileStateViewModel>& InViewModel);
@@ -38,6 +40,9 @@ private:
 	/** Used to build the columns. */
 	TMap<FName, TSharedRef<IFileStateColumnWidgetFactory>> ColumnFactories;
 	
+	/** The text to highlight in the columns. */
+	TAttribute<FText> HighlightText;
+	
 	/** @return The header row to display at the top of the table */
 	TSharedRef<SHeaderRow> MakeHeaderRow();
 	
diff --git a/Engine/Plugins/Developer/Sandbox/SandboxedEditing/Source/SandboxedEditing/Private/Features/Browser/Widgets/Shared/FileState/SFilterableFileStateListView.cpp b/Engine/Plugins/Developer/Sandbox/SandboxedEditing/Source/SandboxedEditing/Private/Features/Browser/Widgets/Shared/FileState/SFilterableFileStateListView.cpp
new file mode 100644
--- /dev/null
+++ b/Engine/Plugins/Developer/Sandbox/SandboxedEditing/Source/SandboxedEditing/Private/Features/Browser/Widgets/Shared/FileState/SFilterableFileStateListView.cpp
@@ -0,0 +1,42 @@
+﻿// Copyright Epic Games, Inc. All Rights Reserved.
+
+#include "SFilterableFileStateListView.h"
+
+#include "SFileStateFilters.h"
+#include "Components/VerticalBox.h"
+#include "Features/Browser/ViewModels/FileState/FilterFileStateViewModel.h"
+#include "Features/Browser/Widgets/Shared/FileState/SFileStateListView.h"
+
+namespace UE::SandboxedEditing
+{
+void SFilterableFileStateListView::Construct(
+	const FArguments& InArgs, 
+	const TSharedRef<IFileStateViewModel>& InFileActionsViewModel, 
+	const TSharedRef<FFilterFileStateViewModel>& InFilterViewModel
+	)
+{
+	ChildSlot
+	[
+		SNew(SVerticalBox)
+
+		+SVerticalBox::Slot()
+		.AutoHeight()
+		.Padding(1.f)
+		[
+			SNew(SFileStateFilters, InFilterViewModel)
+		]
+
+		+SVerticalBox::Slot()
+		.FillHeight(1.f)
+		.Padding(1.f)
+		[
+			SNew(SBorder)
+			[
+				SNew(SFileStateListView, InFileActionsViewModel)
+				.ColumnFactories(InArgs._ColumnFactories)
+				.HighlightText_Lambda([InFilterViewModel]{ return InFilterViewModel->GetSearchText(); })
+			]
+		]
+	];
+}
+}
diff --git a/Engine/Plugins/Developer/Sandbox/SandboxedEditing/Source/SandboxedEditing/Private/Features/Browser/Widgets/Shared/FileState/SFilterableFileStateListView.h b/Engine/Plugins/Developer/Sandbox/SandboxedEditing/Source/SandboxedEditing/Private/Features/Browser/Widgets/Shared/FileState/SFilterableFileStateListView.h
new file mode 100644
--- /dev/null
+++ b/Engine/Plugins/Developer/Sandbox/SandboxedEditing/Source/SandboxedEditing/Private/Features/Browser/Widgets/Shared/FileState/SFilterableFileStateListView.h
@@ -0,0 +1,31 @@
+﻿// Copyright Epic Games, Inc. All Rights Reserved.
+
+#pragma once
+
+#include "Features/Browser/ViewModels/FileState/FileStateColumns.h"
+#include "Widgets/DeclarativeSyntaxSupport.h"
+#include "Widgets/SCompoundWidget.h"
+
+namespace UE::SandboxedEditing
+{
+class FFilterFileStateViewModel;
+class IFileStateViewModel;
+
+/** Vertical arrangement of a filter section (search bar) and SFileStateListView. */
+class SFilterableFileStateListView : public SCompoundWidget
+{
+public:
+	
+	SLATE_BEGIN_ARGS(SFilterableFileStateListView){}
+		/** Used to build the columns. */
+		SLATE_ARGUMENT(FFileStateColumnFactoryMap, ColumnFactories)
+	SLATE_END_ARGS()
+	
+	void Construct(
+		const FArguments& InArgs, 
+		const TSharedRef<IFileStateViewModel>& InFileActionsViewModel, 
+		const TSharedRef<FFilterFileStateViewModel>& InFilterViewModel
+		);
+};
+}
+
diff --git a/Engine/Plugins/Developer/Sandbox/SandboxedEditing/Source/SandboxedEditing/Private/Features/Browser/Widgets/Unsandboxed/SBrowserFilters.cpp b/Engine/Plugins/Developer/Sandbox/SandboxedEditing/Source/SandboxedEditing/Private/Features/Browser/Widgets/Unsandboxed/SBrowserFilters.cpp
--- a/Engine/Plugins/Developer/Sandbox/SandboxedEditing/Source/SandboxedEditing/Private/Features/Browser/Widgets/Unsandboxed/SBrowserFilters.cpp
+++ b/Engine/Plugins/Developer/Sandbox/SandboxedEditing/Source/SandboxedEditing/Private/Features/Browser/Widgets/Unsandboxed/SBrowserFilters.cpp
@@ -10,21 +10,21 @@
 
 namespace UE::SandboxedEditing
 {
-void SBrowserFilters::Construct(const FArguments& InArgs, const TSharedPtr<FFilterSandboxViewModel>& InFilterSandboxViewModel)
+void SBrowserFilters::Construct(const FArguments& InArgs, const TSharedRef<FFilterSandboxViewModel>& InFilterSandboxViewModel)
 {
 	FilterSandboxViewModel = InFilterSandboxViewModel;
 	
 	ChildSlot
 	[
 		SNew(SSearchBox)
-		.HintText(LOCTEXT("SearchHint", "Search Sandbox"))
+		.HintText(LOCTEXT("SearchHint", "Search sandboxes"))
 		.OnTextChanged(this, &SBrowserFilters::OnSearchTextChanged)
 		.OnTextCommitted(this, &SBrowserFilters::OnSearchTextCommitted)
 		.DelayChangeNotificationsWhileTyping(true)
 	];
 }
 
-void SBrowserFilters::OnSearchTextChanged(const FText& Text)
+void SBrowserFilters::OnSearchTextChanged(const FText& Text) const
 {
 	FilterSandboxViewModel->SetSearchText(Text);
 }
diff --git a/Engine/Plugins/Developer/Sandbox/SandboxedEditing/Source/SandboxedEditing/Private/Features/Browser/Widgets/Unsandboxed/SBrowserFilters.h b/Engine/Plugins/Developer/Sandbox/SandboxedEditing/Source/SandboxedEditing/Private/Features/Browser/Widgets/Unsandboxed/SBrowserFilters.h
--- a/Engine/Plugins/Developer/Sandbox/SandboxedEditing/Source/SandboxedEditing/Private/Features/Browser/Widgets/Unsandboxed/SBrowserFilters.h
+++ b/Engine/Plugins/Developer/Sandbox/SandboxedEditing/Source/SandboxedEditing/Private/Features/Browser/Widgets/Unsandboxed/SBrowserFilters.h
@@ -16,14 +16,14 @@ public:
 	SLATE_BEGIN_ARGS(SBrowserFilters) {}
 	SLATE_END_ARGS()
 
-	void Construct(const FArguments& InArgs, const TSharedPtr<FFilterSandboxViewModel>& InFilterSandboxViewModel);
+	void Construct(const FArguments& InArgs, const TSharedRef<FFilterSandboxViewModel>& InFilterSandboxViewModel);
 	
 private:
 	
 	/** Knows how to filter sandbox items. */
 	TSharedPtr<FFilterSandboxViewModel> FilterSandboxViewModel;
 	
-	void OnSearchTextChanged(const FText& Text);
+	void OnSearchTextChanged(const FText& Text) const;
 	void OnSearchTextCommitted(const FText& InFilterText, ETextCommit::Type Arg);
 };
 }
diff --git a/Engine/Plugins/Developer/Sandbox/SandboxedEditing/Source/SandboxedEditing/Private/Features/Browser/Widgets/Unsandboxed/SBrowserSandboxDetails.cpp b/Engine/Plugins/Developer/Sandbox/SandboxedEditing/Source/SandboxedEditing/Private/Features/Browser/Widgets/Unsandboxed/SBrowserSandboxDetails.cpp
--- a/Engine/Plugins/Developer/Sandbox/SandboxedEditing/Source/SandboxedEditing/Private/Features/Browser/Widgets/Unsandboxed/SBrowserSandboxDetails.cpp
+++ b/Engine/Plugins/Developer/Sandbox/SandboxedEditing/Source/SandboxedEditing/Private/Features/Browser/Widgets/Unsandboxed/SBrowserSandboxDetails.cpp
@@ -5,9 +5,10 @@
 #include "Components/VerticalBox.h"
 #include "Features/Browser/ViewModels/FileState/FileStateColumnRegistry.h"
 #include "Features/Browser/ViewModels/FileState/Models/UnloadedSandboxFileStateViewModel.h"
-#include "Features/Browser/Widgets/Shared/FileState/SFileStateListView.h"
+#include "Features/Browser/Widgets/Shared/FileState/SFilterableFileStateListView.h"
 #include "Features/Browser/Widgets/Shared/SSandboxDescription.h"
 #include "Internationalization/Text.h"
+#include "Widgets/Layout/SBorder.h"
 #include "Widgets/Layout/SBox.h"
 #include "Widgets/Layout/SExpandableArea.h"
 #include "Widgets/Layout/SWidgetSwitcher.h"
@@ -35,20 +36,22 @@ static TSharedRef<SWidget> MakeMessageWidget(TAttribute<FText> InMessageAttr)
 void SBrowserSandboxDetails::Construct(
 	const FArguments& InArgs, 
 	const TSharedRef<FSandboxMetaDataViewModel>& InMetaDataViewModel,
-	const TSharedRef<FUnloadedSandboxFileStateViewModel>& InFileActionsViewModel
+	const TSharedRef<FUnloadedSandboxFileStateViewModel>& InFileActionsViewModel,
+	const TSharedRef<FFilterFileStateViewModel>& InFilterViewModel
 	)
 {
 	NumSelectedAttr = InArgs._NumSelected;
 	ChildSlot
 	[
-		MakeWidgetSwitcherContent(InArgs, InMetaDataViewModel, InFileActionsViewModel)
+		MakeWidgetSwitcherContent(InArgs, InMetaDataViewModel, InFileActionsViewModel, InFilterViewModel)
 	];
 }
 
 TSharedRef<SWidget> SBrowserSandboxDetails::MakeWidgetSwitcherContent(
 	const FArguments& InArgs, 
 	const TSharedRef<FSandboxMetaDataViewModel>& InMetaDataViewModel, 
-	const TSharedRef<FUnloadedSandboxFileStateViewModel>& InFileActionsViewModel
+	const TSharedRef<FUnloadedSandboxFileStateViewModel>& InFileActionsViewModel,
+	const TSharedRef<FFilterFileStateViewModel>& InFilterViewModel
 	)
 {
 	return SNew(SBox)
@@ -59,7 +62,7 @@ TSharedRef<SWidget> SBrowserSandboxDetails::MakeWidgetSwitcherContent(
 
 			+SWidgetSwitcher::Slot()
 			[
-				MakeDetailContent(InArgs, InMetaDataViewModel, InFileActionsViewModel)
+				MakeDetailContent(InArgs, InMetaDataViewModel, InFileActionsViewModel, InFilterViewModel)
 			]
 			
 			+SWidgetSwitcher::Slot()
@@ -77,7 +80,8 @@ TSharedRef<SWidget> SBrowserSandboxDetails::MakeWidgetSwitcherContent(
 TSharedRef<SWidget> SBrowserSandboxDetails::MakeDetailContent(
 	const FArguments& InArgs,
 	const TSharedRef<FSandboxMetaDataViewModel>& InMetaDataViewModel, 
-	const TSharedRef<FUnloadedSandboxFileStateViewModel>& InFileActionsViewModel
+	const TSharedRef<FUnloadedSandboxFileStateViewModel>& InFileActionsViewModel,
+	const TSharedRef<FFilterFileStateViewModel>& InFilterViewModel
 	)
 {
 	return SNew(SVerticalBox)
@@ -115,11 +119,8 @@ TSharedRef<SWidget> SBrowserSandboxDetails::MakeDetailContent(
 			]
 			.BodyContent()
 			[
-				SNew(SBorder)
-				[
-					SNew(SFileStateListView, InFileActionsViewModel)
-					.ColumnFactories(InArgs._ColumnFactories)
-				]
+				SNew(SFilterableFileStateListView, InFileActionsViewModel, InFilterViewModel)
+                .ColumnFactories(InArgs._ColumnFactories)
 			]
 		];
 }
diff --git a/Engine/Plugins/Developer/Sandbox/SandboxedEditing/Source/SandboxedEditing/Private/Features/Browser/Widgets/Unsandboxed/SBrowserSandboxDetails.h b/Engine/Plugins/Developer/Sandbox/SandboxedEditing/Source/SandboxedEditing/Private/Features/Browser/Widgets/Unsandboxed/SBrowserSandboxDetails.h
--- a/Engine/Plugins/Developer/Sandbox/SandboxedEditing/Source/SandboxedEditing/Private/Features/Browser/Widgets/Unsandboxed/SBrowserSandboxDetails.h
+++ b/Engine/Plugins/Developer/Sandbox/SandboxedEditing/Source/SandboxedEditing/Private/Features/Browser/Widgets/Unsandboxed/SBrowserSandboxDetails.h
@@ -8,8 +8,9 @@
 
 namespace UE::SandboxedEditing
 {
-class FUnloadedSandboxFileStateViewModel;
+class FFilterFileStateViewModel;
 class FSandboxMetaDataViewModel;
+class FUnloadedSandboxFileStateViewModel;
 
 /** Shows details for a sandbox in the browser widget. Handles showing text such as "select sandbox to view details", etc. */
 class SBrowserSandboxDetails : public SCompoundWidget
@@ -28,7 +29,8 @@ public:
 
 	void Construct(const FArguments& InArgs, 
 		const TSharedRef<FSandboxMetaDataViewModel>& InMetaDataViewModel,
-		const TSharedRef<FUnloadedSandboxFileStateViewModel>& InFileActionsViewModel
+		const TSharedRef<FUnloadedSandboxFileStateViewModel>& InFileActionsViewModel,
+		const TSharedRef<FFilterFileStateViewModel>& InFilterViewModel
 		);
 	
 private:
@@ -40,13 +42,15 @@ private:
 	TSharedRef<SWidget> MakeWidgetSwitcherContent(
 		const FArguments& InArgs, 
 		const TSharedRef<FSandboxMetaDataViewModel>& InMetaDataViewModel,
-		const TSharedRef<FUnloadedSandboxFileStateViewModel>& InFileActionsViewModel
+		const TSharedRef<FUnloadedSandboxFileStateViewModel>& InFileActionsViewModel,
+		const TSharedRef<FFilterFileStateViewModel>& InFilterViewModel
 	);
 	/** @return Content to display in the widget switcher */
 	TSharedRef<SWidget> MakeDetailContent(
 		const FArguments& InArgs, 
 		const TSharedRef<FSandboxMetaDataViewModel>& InMetaDataViewModel, 
-		const TSharedRef<FUnloadedSandboxFileStateViewModel>& InFileActionsViewModel
+		const TSharedRef<FUnloadedSandboxFileStateViewModel>& InFileActionsViewModel,
+		const TSharedRef<FFilterFileStateViewModel>& InFilterViewModel
 		);
 	
 	/** @return Index for widget switcher to display. Determined by how many items are selected. */
diff --git a/Engine/Plugins/Developer/Sandbox/SandboxedEditing/Source/SandboxedEditing/Private/Features/Browser/Widgets/Unsandboxed/SSandboxedEditingBrowser.cpp b/Engine/Plugins/Developer/Sandbox/SandboxedEditing/Source/SandboxedEditing/Private/Features/Browser/Widgets/Unsandboxed/SSandboxedEditingBrowser.cpp
--- a/Engine/Plugins/Developer/Sandbox/SandboxedEditing/Source/SandboxedEditing/Private/Features/Browser/Widgets/Unsandboxed/SSandboxedEditingBrowser.cpp
+++ b/Engine/Plugins/Developer/Sandbox/SandboxedEditing/Source/SandboxedEditing/Private/Features/Browser/Widgets/Unsandboxed/SSandboxedEditingBrowser.cpp
@@ -41,7 +41,9 @@ void SSandboxedEditingBrowser::Construct(const FArguments& InArgs, const FViewMo
 		.SessionContent() [ ListViewRef ]
 		.DetailsContent()
 		[
-			SNew(SBrowserSandboxDetails, InViewModels.MetaDataViewModel, InViewModels.UnloadedSandboxFileStateViewModel)
+			SNew(SBrowserSandboxDetails, 
+				InViewModels.MetaDataViewModel, InViewModels.UnloadedSandboxFileStateViewModel, InViewModels.FilterUnloadedSandboxFileStateViewModel
+				)
 			.SandboxPath(ListViewRef, &SSandboxListView::GetSelectedItem)
 			.NumSelected(ListViewRef, &SSandboxListView::NumSelected)
 			.ColumnFactories(InArgs._FileActionsColumnFactories)
diff --git a/Engine/Plugins/Developer/Sandbox/SandboxedEditing/Source/SandboxedEditing/Private/Features/Browser/Widgets/Unsandboxed/SSandboxedEditingBrowser.h b/Engine/Plugins/Developer/Sandbox/SandboxedEditing/Source/SandboxedEditing/Private/Features/Browser/Widgets/Unsandboxed/SSandboxedEditingBrowser.h
--- a/Engine/Plugins/Developer/Sandbox/SandboxedEditing/Source/SandboxedEditing/Private/Features/Browser/Widgets/Unsandboxed/SSandboxedEditingBrowser.h
+++ b/Engine/Plugins/Developer/Sandbox/SandboxedEditing/Source/SandboxedEditing/Private/Features/Browser/Widgets/Unsandboxed/SSandboxedEditingBrowser.h
@@ -11,6 +11,7 @@ class FUICommandList;
 
 namespace UE::SandboxedEditing
 {
+class FFilterFileStateViewModel;
 class FBrowserCommandBindings;
 class FFilterSandboxViewModel;
 class FSandboxControlsViewModel;
@@ -31,6 +32,7 @@ public:
 		TSharedRef<FSandboxControlsViewModel> ControlsViewModel;
 		TSharedRef<FSandboxMetaDataViewModel> MetaDataViewModel;
 		TSharedRef<FUnloadedSandboxFileStateViewModel> UnloadedSandboxFileStateViewModel;
+		TSharedRef<FFilterFileStateViewModel> FilterUnloadedSandboxFileStateViewModel;
 	};
 
 	SLATE_BEGIN_ARGS(SSandboxedEditingBrowser){}
