Sandboxed Editing: File actions list view now displays:
- "No file changes have been made in the sandbox." -> in browser (and persist dialogue, but should not happen)
- "File changes appear here once you edit, add, or remove a file." -> in active sandbox UI
- "All file actions are filtered." -> when all is filtered

#Synced-CL 49159615

diff --git a/Engine/Plugins/Developer/Sandbox/SandboxedEditing/Source/SandboxedEditing/Private/Features/Browser/ViewModels/FileState/FilterFileStateViewModel.cpp b/Engine/Plugins/Developer/Sandbox/SandboxedEditing/Source/SandboxedEditing/Private/Features/Browser/ViewModels/FileState/FilterFileStateViewModel.cpp
--- a/Engine/Plugins/Developer/Sandbox/SandboxedEditing/Source/SandboxedEditing/Private/Features/Browser/ViewModels/FileState/FilterFileStateViewModel.cpp
+++ b/Engine/Plugins/Developer/Sandbox/SandboxedEditing/Source/SandboxedEditing/Private/Features/Browser/ViewModels/FileState/FilterFileStateViewModel.cpp
@@ -31,6 +31,11 @@ bool FFilterFileStateViewModel::PassesFilter(const TSharedRef<FFileStateItem>& I
 	return TextFilter.PassesFilter(InItem);
 }
 
+bool FFilterFileStateViewModel::AreAnyFiltersActive() const
+{
+	return !GetSearchText().IsEmpty();
+}
+
 void FFilterFileStateViewModel::ItemToString(const TSharedRef<FFileStateItem>& InItem, TArray<FString>& OutSearchTerms)
 {
 	for (const TSharedRef<IFileStateColumnBehavior>& Column : VisibleColumns.Get())
diff --git a/Engine/Plugins/Developer/Sandbox/SandboxedEditing/Source/SandboxedEditing/Private/Features/Browser/ViewModels/FileState/FilterFileStateViewModel.h b/Engine/Plugins/Developer/Sandbox/SandboxedEditing/Source/SandboxedEditing/Private/Features/Browser/ViewModels/FileState/FilterFileStateViewModel.h
--- a/Engine/Plugins/Developer/Sandbox/SandboxedEditing/Source/SandboxedEditing/Private/Features/Browser/ViewModels/FileState/FilterFileStateViewModel.h
+++ b/Engine/Plugins/Developer/Sandbox/SandboxedEditing/Source/SandboxedEditing/Private/Features/Browser/ViewModels/FileState/FilterFileStateViewModel.h
@@ -28,6 +28,9 @@ public:
 	/** @return Whether the item can be displayed. */
 	bool PassesFilter(const TSharedRef<FFileStateItem>& InItem) const;
 	
+	/** @return Whether any filters are active. */
+	bool AreAnyFiltersActive() const;
+	
 	/** @return Delegate invoked when any of the filters changes */
 	FSimpleMulticastDelegate& OnFilterChanged() { return OnFilterChangedDelegate; }
 	
diff --git a/Engine/Plugins/Developer/Sandbox/SandboxedEditing/Source/SandboxedEditing/Private/Features/Browser/Widgets/Sandboxed/SActiveSandbox.cpp b/Engine/Plugins/Developer/Sandbox/SandboxedEditing/Source/SandboxedEditing/Private/Features/Browser/Widgets/Sandboxed/SActiveSandbox.cpp
--- a/Engine/Plugins/Developer/Sandbox/SandboxedEditing/Source/SandboxedEditing/Private/Features/Browser/Widgets/Sandboxed/SActiveSandbox.cpp
+++ b/Engine/Plugins/Developer/Sandbox/SandboxedEditing/Source/SandboxedEditing/Private/Features/Browser/Widgets/Sandboxed/SActiveSandbox.cpp
@@ -97,6 +97,7 @@ TSharedRef<SWidget> SActiveSandbox::MakeFileChangesContent(const FArguments& InA
 		[
 			SNew(SFilterableFileStateListView, InViewModels.ActiveSandboxFileStateModel, InViewModels.FilterActiveSandboxFileStateModel)
 			.ColumnFactories(InArgs._ColumnFactories)
+			.NoChangesText(LOCTEXT("NoItems", "File changes appear here once you edit, add, or remove a file."))
 		];
 }
 }
diff --git a/Engine/Plugins/Developer/Sandbox/SandboxedEditing/Source/SandboxedEditing/Private/Features/Browser/Widgets/Shared/FileState/SFileStateListView.cpp b/Engine/Plugins/Developer/Sandbox/SandboxedEditing/Source/SandboxedEditing/Private/Features/Browser/Widgets/Shared/FileState/SFileStateListView.cpp
--- a/Engine/Plugins/Developer/Sandbox/SandboxedEditing/Source/SandboxedEditing/Private/Features/Browser/Widgets/Shared/FileState/SFileStateListView.cpp
+++ b/Engine/Plugins/Developer/Sandbox/SandboxedEditing/Source/SandboxedEditing/Private/Features/Browser/Widgets/Shared/FileState/SFileStateListView.cpp
@@ -2,12 +2,16 @@
 
 #include "SFileStateListView.h"
 
-#include "SandboxedEditingStyle.h"
-#include "SFileStateListRow.h"
 #include "Features/Browser/ViewModels/FileState/IFileStateColumnWidgetFactory.h"
 #include "Features/Browser/ViewModels/FileState/IFileStateViewModel.h"
+#include "SFileStateListRow.h"
+#include "SandboxedEditingStyle.h"
+#include "Widgets/SBoxPanel.h"
+#include "Widgets/SOverlay.h"
 #include "Widgets/Views/SHeaderRow.h"
 
+#define LOCTEXT_NAMESPACE "SFileStateListView"
+
 namespace UE::SandboxedEditing
 {
 void SFileStateListView::Construct(const FArguments& InArgs, const TSharedRef<IFileStateViewModel>& InViewModel)
@@ -15,14 +19,35 @@ void SFileStateListView::Construct(const FArguments& InArgs, const TSharedRef<IF
 	ViewModel = InViewModel;
 	ColumnFactories = InArgs._ColumnFactories;
 	HighlightText = InArgs._HighlightText;
+	OverrideOverlayText = InArgs._OverrideOverlayText;
+	NoChangesText = InArgs._NoChangesText.IsEmpty() 
+		? LOCTEXT("NoItemsDefault", "No file changes have been made in the sandbox.") : InArgs._NoChangesText;
 	
 	ChildSlot
 	[
-		SAssignNew(ListView, SListView<TSharedPtr<FFileStateItem>>)
-		.ListItemsSource(&ViewModel->GetItems())
-		.OnGenerateRow(this, &SFileStateListView::OnGenerateRow)
-		.SelectionMode(ESelectionMode::None)
-		.HeaderRow(MakeHeaderRow())
+		SNew(SOverlay)
+
+		+SOverlay::Slot()
+		[
+			SAssignNew(ListView, SListView<TSharedPtr<FFileStateItem>>)
+			.ListItemsSource(&ViewModel->GetItems())
+			.OnGenerateRow(this, &SFileStateListView::OnGenerateRow)
+			.SelectionMode(ESelectionMode::None)
+			.HeaderRow(MakeHeaderRow())
+		]
+
+		+SOverlay::Slot()
+		.Padding(0, 20, 0, 0) // To ensure the panel is below the row header.
+		[
+			SNew(SBox)
+			.HAlign(HAlign_Center)
+			.VAlign(VAlign_Center)
+			[
+				SNew(STextBlock)
+				.Text(this, &SFileStateListView::GetOverlayText)
+				.Visibility(this, &SFileStateListView::GetOverlayTextVisibility)
+			]
+		]
 	];
 	
 	ViewModel->OnItemsChanged().AddSP(ListView.ToSharedRef(), &SListView<TSharedPtr<FFileStateItem>>::RequestListRefresh);
@@ -53,4 +78,24 @@ TSharedRef<ITableRow> SFileStateListView::OnGenerateRow(TSharedPtr<FFileStateIte
 		.Padding(Padding)
 		.HighlightText(HighlightText);
 }
+
+FText SFileStateListView::GetOverlayText() const
+{
+	const FText Override = OverrideOverlayText.IsBound() || OverrideOverlayText.IsSet() ? OverrideOverlayText.Get() : FText::GetEmpty();
+	if (!Override.IsEmpty())
+	{
+		return Override;
+	}
+	
+	return ViewModel->GetItems().IsEmpty()
+		? NoChangesText
+		: FText::GetEmpty();
+}
+
+EVisibility SFileStateListView::GetOverlayTextVisibility() const
+{
+	return GetOverlayText().IsEmpty() ? EVisibility::Collapsed : EVisibility::HitTestInvisible;
 }
+}
+
+#undef LOCTEXT_NAMESPACE
\ No newline at end of file
diff --git a/Engine/Plugins/Developer/Sandbox/SandboxedEditing/Source/SandboxedEditing/Private/Features/Browser/Widgets/Shared/FileState/SFileStateListView.h b/Engine/Plugins/Developer/Sandbox/SandboxedEditing/Source/SandboxedEditing/Private/Features/Browser/Widgets/Shared/FileState/SFileStateListView.h
--- a/Engine/Plugins/Developer/Sandbox/SandboxedEditing/Source/SandboxedEditing/Private/Features/Browser/Widgets/Shared/FileState/SFileStateListView.h
+++ b/Engine/Plugins/Developer/Sandbox/SandboxedEditing/Source/SandboxedEditing/Private/Features/Browser/Widgets/Shared/FileState/SFileStateListView.h
@@ -20,13 +20,23 @@ class SFileStateListView : public SCompoundWidget
 {
 public:
 	
-	SLATE_BEGIN_ARGS(SFileStateListView) {}
+	SLATE_BEGIN_ARGS(SFileStateListView)
+		{}
 		/** Used to build the columns. */
 		SLATE_ARGUMENT(FFileStateColumnFactoryMap, ColumnFactories)
+		
 		/** The text to highlight in the columns. */
 		SLATE_ATTRIBUTE(FText, HighlightText)
+		
+		/** Text to display in the content, e.g. when all items are filtered out. Overrides the text when the attribute returns non-empty. */
+		SLATE_ATTRIBUTE(FText, OverrideOverlayText)
+		/**
+		 * The text to display when there are no changes to display.
+		 * Defaults to: "No file changes have been made in the sandbox."
+		 */
+		SLATE_ARGUMENT(FText, NoChangesText)
 	SLATE_END_ARGS()
-	
+
 	void Construct(const FArguments& InArgs, const TSharedRef<IFileStateViewModel>& InViewModel);
 	
 private:
@@ -42,12 +52,20 @@ private:
 	
 	/** The text to highlight in the columns. */
 	TAttribute<FText> HighlightText;
+	/** Text to display in the content, e.g. when all items are filtered out. Overrides the text when the attribute returns non-empty. */
+	TAttribute<FText> OverrideOverlayText;
+	/** The text to display when there are no changes to display. */
+	FText NoChangesText;
 	
 	/** @return The header row to display at the top of the table */
 	TSharedRef<SHeaderRow> MakeHeaderRow();
 	
 	/** @return Widget for the items */
 	TSharedRef<ITableRow> OnGenerateRow(TSharedPtr<FFileStateItem> InItem, const TSharedRef<STableViewBase>& InOwnerTable) const;
+	
+	/** @return Text to display in the content, e.g. when there's no items. */
+	FText GetOverlayText() const;
+	/** @return The visibility of the overlay text */
+	EVisibility GetOverlayTextVisibility() const;
 };
-}
-
+}
\ No newline at end of file
diff --git a/Engine/Plugins/Developer/Sandbox/SandboxedEditing/Source/SandboxedEditing/Private/Features/Browser/Widgets/Shared/FileState/SFilterableFileStateListView.cpp b/Engine/Plugins/Developer/Sandbox/SandboxedEditing/Source/SandboxedEditing/Private/Features/Browser/Widgets/Shared/FileState/SFilterableFileStateListView.cpp
--- a/Engine/Plugins/Developer/Sandbox/SandboxedEditing/Source/SandboxedEditing/Private/Features/Browser/Widgets/Shared/FileState/SFilterableFileStateListView.cpp
+++ b/Engine/Plugins/Developer/Sandbox/SandboxedEditing/Source/SandboxedEditing/Private/Features/Browser/Widgets/Shared/FileState/SFilterableFileStateListView.cpp
@@ -5,8 +5,11 @@
 #include "SFileStateFilters.h"
 #include "Components/VerticalBox.h"
 #include "Features/Browser/ViewModels/FileState/FilterFileStateViewModel.h"
+#include "Features/Browser/ViewModels/FileState/IFileStateViewModel.h"
 #include "Features/Browser/Widgets/Shared/FileState/SFileStateListView.h"
 
+#define LOCTEXT_NAMESPACE "SFilterableFileStateListView"
+
 namespace UE::SandboxedEditing
 {
 void SFilterableFileStateListView::Construct(
@@ -15,6 +18,9 @@ void SFilterableFileStateListView::Construct(
 	const TSharedRef<FFilterFileStateViewModel>& InFilterViewModel
 	)
 {
+	FileActionViewModel = InFileActionsViewModel;
+	FilterViewModel = InFilterViewModel;
+	
 	ChildSlot
 	[
 		SNew(SVerticalBox)
@@ -35,8 +41,20 @@ void SFilterableFileStateListView::Construct(
 				SNew(SFileStateListView, InFileActionsViewModel)
 				.ColumnFactories(InArgs._ColumnFactories)
 				.HighlightText_Lambda([InFilterViewModel]{ return InFilterViewModel->GetSearchText(); })
+				.OverrideOverlayText(this, &SFilterableFileStateListView::GetOverrideOverlayText)
+				.NoChangesText(InArgs._NoChangesText)
 			]
 		]
 	];
 }
+
+FText SFilterableFileStateListView::GetOverrideOverlayText() const
+{
+	const bool bAllFilteredOut = FilterViewModel->AreAnyFiltersActive() && FileActionViewModel->GetItems().IsEmpty();
+	return bAllFilteredOut 
+		? LOCTEXT("AllFiltered", "All file actions are filtered.")
+		: FText::GetEmpty();
 }
+}
+
+#undef LOCTEXT_NAMESPACE
\ No newline at end of file
diff --git a/Engine/Plugins/Developer/Sandbox/SandboxedEditing/Source/SandboxedEditing/Private/Features/Browser/Widgets/Shared/FileState/SFilterableFileStateListView.h b/Engine/Plugins/Developer/Sandbox/SandboxedEditing/Source/SandboxedEditing/Private/Features/Browser/Widgets/Shared/FileState/SFilterableFileStateListView.h
--- a/Engine/Plugins/Developer/Sandbox/SandboxedEditing/Source/SandboxedEditing/Private/Features/Browser/Widgets/Shared/FileState/SFilterableFileStateListView.h
+++ b/Engine/Plugins/Developer/Sandbox/SandboxedEditing/Source/SandboxedEditing/Private/Features/Browser/Widgets/Shared/FileState/SFilterableFileStateListView.h
@@ -19,13 +19,26 @@ public:
 	SLATE_BEGIN_ARGS(SFilterableFileStateListView){}
 		/** Used to build the columns. */
 		SLATE_ARGUMENT(FFileStateColumnFactoryMap, ColumnFactories)
+
+		/** The text to display when there are no changes to display. */
+		SLATE_ARGUMENT(FText, NoChangesText)
 	SLATE_END_ARGS()
-	
+
 	void Construct(
 		const FArguments& InArgs, 
 		const TSharedRef<IFileStateViewModel>& InFileActionsViewModel, 
 		const TSharedRef<FFilterFileStateViewModel>& InFilterViewModel
 		);
+	
+private:
+	
+	/** Used to generate the override overlay text. */
+	TSharedPtr<IFileStateViewModel> FileActionViewModel;
+	/** Used to generate the override overlay text. */
+	TSharedPtr<FFilterFileStateViewModel> FilterViewModel;
+	
+	/** @return Text indicating that all items are filtered out, or empty. */
+	FText GetOverrideOverlayText() const;
 };
 }
 
