Sandboxed Editing: Browser sandbox list view now shows "Sandboxes you create appear here." when there are no sandboxes.

#Synced-CL 49159615

diff --git a/Engine/Plugins/Developer/Sandbox/SandboxedEditing/Source/SandboxedEditing/Private/Features/Browser/Widgets/Unsandboxed/List/SSandboxListView.cpp b/Engine/Plugins/Developer/Sandbox/SandboxedEditing/Source/SandboxedEditing/Private/Features/Browser/Widgets/Unsandboxed/List/SSandboxListView.cpp
--- a/Engine/Plugins/Developer/Sandbox/SandboxedEditing/Source/SandboxedEditing/Private/Features/Browser/Widgets/Unsandboxed/List/SSandboxListView.cpp
+++ b/Engine/Plugins/Developer/Sandbox/SandboxedEditing/Source/SandboxedEditing/Private/Features/Browser/Widgets/Unsandboxed/List/SSandboxListView.cpp
@@ -2,18 +2,20 @@
 
 #include "SSandboxListView.h"
 
+#include "Features/Browser/Commands/BrowserCommands.h"
 #include "Features/Browser/ViewModels/List/ISandboxColumnWidgetFactory.h"
 #include "Features/Browser/ViewModels/List/SandboxListItem.h"
 #include "Features/Browser/ViewModels/List/SandboxListViewModel.h"
 #include "Features/Browser/ViewModels/SandboxControlsViewModel.h"
 #include "Features/Browser/Widgets/Unsandboxed/List/SCreateSandboxListRow.h"
-#include "SSandboxListRow.h"
-#include "SandboxedEditingStyle.h"
-#include "Features/Browser/Commands/BrowserCommands.h"
 #include "Framework/Commands/GenericCommands.h"
 #include "Framework/MultiBox/MultiBoxBuilder.h"
+#include "SSandboxListRow.h"
+#include "SandboxedEditingStyle.h"
 #include "Utils/BreakBehavior.h"
 #include "Widgets/Views/SListView.h"
+#include "Widgets/SBoxPanel.h"
+#include "Widgets/SOverlay.h"
 
 #define LOCTEXT_NAMESPACE "SSandboxListView"
 
@@ -30,14 +32,32 @@ void SSandboxListView::Construct(const FArguments& InArgs, const FViewModels& In
 	
 	ChildSlot
 	[
-		SAssignNew(ListView, SListView<TSharedPtr<FSandboxListItem>>)
-		.ListItemsSource(&Items)
-		.SelectionMode(ESelectionMode::Multi)
-		.OnGenerateRow(this, &SSandboxListView::OnGenerateRow)
-		.OnMouseButtonDoubleClick(this, &SSandboxListView::OnItemDoubleClicked)
-		.HeaderRow(MakeHeaderRow())
-		.OnSelectionChanged(this, &SSandboxListView::OnSelectionChanged)
-		.OnContextMenuOpening(this, &SSandboxListView::OnContextMenuOpening)
+		SNew(SOverlay)
+		
+		+SOverlay::Slot()
+		[
+			SAssignNew(ListView, SListView<TSharedPtr<FSandboxListItem>>)
+			.ListItemsSource(&Items)
+			.SelectionMode(ESelectionMode::Multi)
+			.OnGenerateRow(this, &SSandboxListView::OnGenerateRow)
+			.OnMouseButtonDoubleClick(this, &SSandboxListView::OnItemDoubleClicked)
+			.HeaderRow(MakeHeaderRow())
+			.OnSelectionChanged(this, &SSandboxListView::OnSelectionChanged)
+			.OnContextMenuOpening(this, &SSandboxListView::OnContextMenuOpening)
+		]
+
+		+SOverlay::Slot()
+		.Padding(0, 20, 0, 0) // To ensure the panel is below the row header.
+		[
+			SNew(SBox)
+			.HAlign(HAlign_Center)
+			.VAlign(VAlign_Center)
+			[
+				SNew(STextBlock)
+				.Text(LOCTEXT("NoSandboxes", "Sandboxes you create appear here."))
+				.Visibility(this, &SSandboxListView::GetOverlayTextVisibility)
+			]
+		]
 	];
 	
 	ControlsViewModel->OnCreationWorkflowStarted().AddSP(this, &SSandboxListView::OnCreationWorkflowStarted);
@@ -185,6 +205,11 @@ void SSandboxListView::OnSelectionChanged(TSharedPtr<FSandboxListItem> InItem, E
 		InItem ? TOptional<FString>(InItem->SandboxInfo.SandboxRoot) : TOptional<FString>()
 		);
 }
+
+EVisibility SSandboxListView::GetOverlayTextVisibility() const
+{
+	return ListViewModel->GetItems().IsEmpty() ? EVisibility::HitTestInvisible : EVisibility::Collapsed;
+}
 }
 
 #undef LOCTEXT_NAMESPACE
\ No newline at end of file
diff --git a/Engine/Plugins/Developer/Sandbox/SandboxedEditing/Source/SandboxedEditing/Private/Features/Browser/Widgets/Unsandboxed/List/SSandboxListView.h b/Engine/Plugins/Developer/Sandbox/SandboxedEditing/Source/SandboxedEditing/Private/Features/Browser/Widgets/Unsandboxed/List/SSandboxListView.h
--- a/Engine/Plugins/Developer/Sandbox/SandboxedEditing/Source/SandboxedEditing/Private/Features/Browser/Widgets/Unsandboxed/List/SSandboxListView.h
+++ b/Engine/Plugins/Developer/Sandbox/SandboxedEditing/Source/SandboxedEditing/Private/Features/Browser/Widgets/Unsandboxed/List/SSandboxListView.h
@@ -104,6 +104,8 @@ private:
 	void OnCreationWorkflowEnded();
 	
 	void OnSelectionChanged(TSharedPtr<FSandboxListItem> InItem, ESelectInfo::Type) const;
+	
+	EVisibility GetOverlayTextVisibility() const;
 };
 }
 
