Sandboxed Editing: Add timestamp column.

#Synced-CL 49159615

diff --git a/Engine/Plugins/Developer/Sandbox/SandboxedEditing/Source/SandboxedEditing/Private/Features/Browser/ViewModels/FileState/FileStateColumnRegistry.cpp b/Engine/Plugins/Developer/Sandbox/SandboxedEditing/Source/SandboxedEditing/Private/Features/Browser/ViewModels/FileState/FileStateColumnRegistry.cpp
--- a/Engine/Plugins/Developer/Sandbox/SandboxedEditing/Source/SandboxedEditing/Private/Features/Browser/ViewModels/FileState/FileStateColumnRegistry.cpp
+++ b/Engine/Plugins/Developer/Sandbox/SandboxedEditing/Source/SandboxedEditing/Private/Features/Browser/ViewModels/FileState/FileStateColumnRegistry.cpp
@@ -6,6 +6,7 @@
 #include "Features/Browser/Widgets/Shared/FileState/Columns/ActionFileStateColumn.h"
 #include "Features/Browser/Widgets/Shared/FileState/Columns/PathFileStateColumn.h"
 #include "Features/Browser/Widgets/Shared/FileState/Columns/PersistFileStateColumn.h"
+#include "Features/Browser/Widgets/Shared/FileState/Columns/TimestampFileStateColumn.h"
 
 namespace UE::SandboxedEditing
 {
@@ -14,17 +15,20 @@ class FPathFileStateColumn;
 FFileStateColumnRegistry GetColumnsForBrowser()
 {
 	const TSharedRef<FActionFileStateColumn> ActionColum = MakeShared<FActionFileStateColumn>();
+	const TSharedRef<FTimestampFileStateColumn> Timestamp = MakeShared<FTimestampFileStateColumn>();
 	const TSharedRef<FPathFileStateColumn> PathColumn = MakeShared<FPathFileStateColumn>();
 	
 	FFileStateColumnRegistry Registry;
 	Registry.ColumnBehaviors = TMap<FName, TSharedRef<IFileStateColumnBehavior>>
 	{
 		{ FileActionColumn, ActionColum },
+		{ FileTimestampColumn, Timestamp },
 		{ FilePathColumn, PathColumn }
 	};
 	Registry.ColumnFactories = TMap<FName, TSharedRef<IFileStateColumnWidgetFactory>>
 	{
 		{ FileActionColumn, ActionColum },
+		{ FileTimestampColumn, Timestamp },
 		{ FilePathColumn, PathColumn }
 	};
 	return Registry;
@@ -33,18 +37,21 @@ FFileStateColumnRegistry GetColumnsForBrowser()
 FFileStateColumnRegistry GetColumnsForActiveSandbox()
 {
 	const TSharedRef<FActionFileStateColumn> ActionColum = MakeShared<FActionFileStateColumn>();
+	const TSharedRef<FTimestampFileStateColumn> Timestamp = MakeShared<FTimestampFileStateColumn>();
 	const TSharedRef<FPathFileStateColumn> PathColumn = MakeShared<FPathFileStateColumn>();
 	
 	FFileStateColumnRegistry Registry;
 	Registry.ColumnBehaviors = TMap<FName, TSharedRef<IFileStateColumnBehavior>>
 	{
 		{ FileActionColumn, ActionColum },
-		{ FilePathColumn, PathColumn }
+		{ FileTimestampColumn, Timestamp },
+		{ FilePathColumn, PathColumn },
 	};
 	Registry.ColumnFactories = TMap<FName, TSharedRef<IFileStateColumnWidgetFactory>>
 	{
 		{ FileActionColumn, ActionColum },
-		{ FilePathColumn, PathColumn }
+		{ FileTimestampColumn, Timestamp },
+		{ FilePathColumn, PathColumn },
 	};
 	return Registry;
 }
@@ -53,6 +60,7 @@ FFileStateColumnRegistry GetColumnsForPersist(const TSharedRef<FPersistOperation
 {
 	const TSharedRef<FPersistFileStateColumn> CheckboxColumn = MakeShared<FPersistFileStateColumn>(InViewModel);
 	const TSharedRef<FActionFileStateColumn> ActionColum = MakeShared<FActionFileStateColumn>();
+	const TSharedRef<FTimestampFileStateColumn> Timestamp = MakeShared<FTimestampFileStateColumn>();
 	const TSharedRef<FPathFileStateColumn> PathColumn = MakeShared<FPathFileStateColumn>();
 	
 	FFileStateColumnRegistry Registry;
@@ -60,13 +68,15 @@ FFileStateColumnRegistry GetColumnsForPersist(const TSharedRef<FPersistOperation
 	{
 		{ PersistCheckboxColumn, CheckboxColumn },
 		{ FileActionColumn, ActionColum },
-		{ FilePathColumn, PathColumn }
+		{ FileTimestampColumn, Timestamp },
+		{ FilePathColumn, PathColumn },
 	};
 	Registry.ColumnFactories = TMap<FName, TSharedRef<IFileStateColumnWidgetFactory>>
 	{
 		{ PersistCheckboxColumn, CheckboxColumn },
 		{ FileActionColumn, ActionColum },
-		{ FilePathColumn, PathColumn }
+		{ FileTimestampColumn, Timestamp },
+		{ FilePathColumn, PathColumn },
 	};
 	return Registry;
 }
diff --git a/Engine/Plugins/Developer/Sandbox/SandboxedEditing/Source/SandboxedEditing/Private/Features/Browser/ViewModels/FileState/FileStateColumns.h b/Engine/Plugins/Developer/Sandbox/SandboxedEditing/Source/SandboxedEditing/Private/Features/Browser/ViewModels/FileState/FileStateColumns.h
--- a/Engine/Plugins/Developer/Sandbox/SandboxedEditing/Source/SandboxedEditing/Private/Features/Browser/ViewModels/FileState/FileStateColumns.h
+++ b/Engine/Plugins/Developer/Sandbox/SandboxedEditing/Source/SandboxedEditing/Private/Features/Browser/ViewModels/FileState/FileStateColumns.h
@@ -17,4 +17,5 @@ using FFileStateColumnFactoryMap = TMap<FName, TSharedRef<IFileStateColumnWidget
 constexpr FLazyName PersistCheckboxColumn(TEXT("PersistCheckboxColumn"));
 constexpr FLazyName FileActionColumn(TEXT("FileActionColumn"));
 constexpr FLazyName FilePathColumn(TEXT("FilePathColumn"));
+constexpr FLazyName FileTimestampColumn(TEXT("FileTimestampColumn"));
 }
diff --git a/Engine/Plugins/Developer/Sandbox/SandboxedEditing/Source/SandboxedEditing/Private/Features/Browser/ViewModels/FileState/FileStateItem.h b/Engine/Plugins/Developer/Sandbox/SandboxedEditing/Source/SandboxedEditing/Private/Features/Browser/ViewModels/FileState/FileStateItem.h
--- a/Engine/Plugins/Developer/Sandbox/SandboxedEditing/Source/SandboxedEditing/Private/Features/Browser/ViewModels/FileState/FileStateItem.h
+++ b/Engine/Plugins/Developer/Sandbox/SandboxedEditing/Source/SandboxedEditing/Private/Features/Browser/ViewModels/FileState/FileStateItem.h
@@ -3,6 +3,7 @@
 #pragma once
 
 #include "Containers/UnrealString.h"
+#include "Misc/DateTime.h"
 #include "Types/SandboxFileChange.h"
 
 namespace UE::SandboxedEditing
@@ -15,10 +16,17 @@ struct FFileStateItem
 	
 	/** The operation performed on the file. */
 	FileSandboxCore::ESandboxFileChange Action;
+	
+	/** The time of the edit */
+	FDateTime Timestamp;
 
-	explicit FFileStateItem(FString InNonSandboxFile, FileSandboxCore::ESandboxFileChange InAction)
+	explicit FFileStateItem(FString InNonSandboxFile, FileSandboxCore::ESandboxFileChange InAction, const FDateTime& InTimestamp)
 		: NonSandboxFile(MoveTemp(InNonSandboxFile))
 		, Action(InAction)
+		, Timestamp(InTimestamp)
 	{}
+	
+	/** @return Whether the timestamp is valid. FileSandboxCore API returns FDateTime::MinValue for invalid timestamps. */
+	bool HasValidTimestamp() const { return Timestamp != FDateTime::MinValue(); }
 };
 }
diff --git a/Engine/Plugins/Developer/Sandbox/SandboxedEditing/Source/SandboxedEditing/Private/Features/Browser/ViewModels/FileState/FileStateItemUtils.cpp b/Engine/Plugins/Developer/Sandbox/SandboxedEditing/Source/SandboxedEditing/Private/Features/Browser/ViewModels/FileState/FileStateItemUtils.cpp
--- a/Engine/Plugins/Developer/Sandbox/SandboxedEditing/Source/SandboxedEditing/Private/Features/Browser/ViewModels/FileState/FileStateItemUtils.cpp
+++ b/Engine/Plugins/Developer/Sandbox/SandboxedEditing/Source/SandboxedEditing/Private/Features/Browser/ViewModels/FileState/FileStateItemUtils.cpp
@@ -13,12 +13,13 @@ TArray<TSharedPtr<FFileStateItem>> MakeItemsFromFileChanges(const FileSandboxCor
 {
 	const TArray<FString>& NonSandboxPaths = InChanges.NonSandboxPaths;
 	const TArray<FileSandboxCore::ESandboxFileChange>& FileActions = InChanges.FileActions;
-	check(NonSandboxPaths.Num() == FileActions.Num());
+	const TArray<FDateTime>& Timestamps = InChanges.Timestamps;
+	check(NonSandboxPaths.Num() == FileActions.Num() && NonSandboxPaths.Num() == Timestamps.Num());
 	
 	TArray<TSharedPtr<FFileStateItem>> Items;
 	for (int32 Index = 0; Index < NonSandboxPaths.Num(); ++Index)
 	{
-		Items.Emplace(MakeShared<FFileStateItem>(NonSandboxPaths[Index], FileActions[Index]));
+		Items.Emplace(MakeShared<FFileStateItem>(NonSandboxPaths[Index], FileActions[Index], Timestamps[Index]));
 	}
 	return Items;
 }
diff --git a/Engine/Plugins/Developer/Sandbox/SandboxedEditing/Source/SandboxedEditing/Private/Features/Browser/ViewModels/FileState/Models/ActiveSandboxFileStateViewModel.cpp b/Engine/Plugins/Developer/Sandbox/SandboxedEditing/Source/SandboxedEditing/Private/Features/Browser/ViewModels/FileState/Models/ActiveSandboxFileStateViewModel.cpp
--- a/Engine/Plugins/Developer/Sandbox/SandboxedEditing/Source/SandboxedEditing/Private/Features/Browser/ViewModels/FileState/Models/ActiveSandboxFileStateViewModel.cpp
+++ b/Engine/Plugins/Developer/Sandbox/SandboxedEditing/Source/SandboxedEditing/Private/Features/Browser/ViewModels/FileState/Models/ActiveSandboxFileStateViewModel.cpp
@@ -58,15 +58,19 @@ void FActiveSandboxFileStateViewModel::RefreshItems()
 			{
 				return FPaths::IsSamePath(Item->NonSandboxFile, InNonSandboxPath);
 			});
+			const FDateTime Timestamp = Model->GetSandboxedFileTimestamp(InNonSandboxPath);
 			if (ReusedItem)
 			{
-				ReusedItem->Get()->Action = InAction; // Do this to make sure the item is up to date
+				// Do this to make sure the item is up to date
+				ReusedItem->Get()->Action = InAction; 
+				ReusedItem->Get()->Timestamp = Timestamp;
+				
 				AllUnfilteredItems.Emplace(*ReusedItem);
 			}
 			else
 			{
 				bMadeChanges = true;
-				AllUnfilteredItems.Emplace(MakeShared<FFileStateItem>(InNonSandboxPath, InAction));
+				AllUnfilteredItems.Emplace(MakeShared<FFileStateItem>(InNonSandboxPath, InAction, Timestamp));
 			}
 				
 			return FileSandboxCore::EBreakBehavior::Continue;
diff --git a/Engine/Plugins/Developer/Sandbox/SandboxedEditing/Source/SandboxedEditing/Private/Features/Browser/Widgets/Shared/FileState/Columns/TimestampFileStateColumn.cpp b/Engine/Plugins/Developer/Sandbox/SandboxedEditing/Source/SandboxedEditing/Private/Features/Browser/Widgets/Shared/FileState/Columns/TimestampFileStateColumn.cpp
new file mode 100644
--- /dev/null
+++ b/Engine/Plugins/Developer/Sandbox/SandboxedEditing/Source/SandboxedEditing/Private/Features/Browser/Widgets/Shared/FileState/Columns/TimestampFileStateColumn.cpp
@@ -0,0 +1,90 @@
+﻿// Copyright Epic Games, Inc. All Rights Reserved.
+
+#include "TimestampFileStateColumn.h"
+
+#include "SandboxedEditingStyle.h"
+#include "Features/Browser/ViewModels/FileState/FileStateColumns.h"
+#include "Features/Browser/ViewModels/FileState/FileStateItem.h"
+#include "Widgets/Layout/SBox.h"
+#include "Widgets/Text/STextBlock.h"
+
+#define LOCTEXT_NAMESPACE "FTimestampFileStateColumn"
+
+namespace UE::SandboxedEditing
+{
+namespace TimestampDetail
+{
+static FText FormatRelativeTime(const FDateTime& EventTime, const FDateTime* CurrTime = nullptr)
+{
+	FTimespan TimeSpan = (CurrTime ? *CurrTime : FDateTime::Now()) - EventTime;
+	int32 Days = TimeSpan.GetDays();
+	int32 Hours = TimeSpan.GetHours();
+
+	if (Days >= 1)
+	{
+		return Hours > 0 ?
+			FText::Format(LOCTEXT("DaysHours", "{0} {0}|plural(one=Day,other=Days), {1} {1}|plural(one=Hour,other=Hours) Ago"), Days, Hours) :
+			FText::Format(LOCTEXT("Days", "{0} {0}|plural(one=Day,other=Days) Ago"), Days);
+	}
+
+	int32 Minutes = TimeSpan.GetMinutes();
+	if (Hours >= 1)
+	{
+		return Minutes > 0 ?
+			FText::Format(LOCTEXT("HoursMins", "{0} {0}|plural(one=Hour,other=Hours), {1} {1}|plural(one=Minute,other=Minutes) Ago"), Hours, Minutes) :
+			FText::Format(LOCTEXT("Hours", "{0} {0}|plural(one=Hour,other=Hours) Ago"), Hours);
+	}
+
+	int32 Seconds = TimeSpan.GetSeconds();
+	if (Minutes >= 1)
+	{
+		return Seconds > 0 ?
+			FText::Format(LOCTEXT("MinsSecs", "{0} {0}|plural(one=Minute,other=Minutes), {1} {1}|plural(one=Second,other=Seconds) Ago"), Minutes, Seconds) :
+			FText::Format(LOCTEXT("Mins", "{0} {0}|plural(one=Minute,other=Minutes) Ago"), Minutes, Seconds);
+	}
+
+	if (Seconds >= 1)
+	{
+		return FText::Format(LOCTEXT("Secs", "{0} {0}|plural(one=Second,other=Seconds) Ago"), Seconds);
+	}
+	return LOCTEXT("Now", "Now");
+}
+
+static FText GetRelativeTimestamp(const FFileStateItem& InItem)
+{
+	return InItem.HasValidTimestamp()
+		? FormatRelativeTime(InItem.Timestamp)
+		: LOCTEXT("NotApplicable", "n/a");
+}
+}
+
+void FTimestampFileStateColumn::PopulateSearchTerms(const TSharedPtr<FFileStateItem>& InRowData, TArray<FString>& OutSearchTerms) const
+{
+	OutSearchTerms.Add(
+		TimestampDetail::GetRelativeTimestamp(*InRowData).ToString()
+		);
+}
+
+SHeaderRow::FColumn::FArguments FTimestampFileStateColumn::MakeColumnArguments()
+{
+	return SHeaderRow::FColumn::FArguments()
+		.ColumnId(FileTimestampColumn)
+		.FixedWidth(FSandboxedEditingStyle::Get().GetFloat("SandboxedEditing.FileActions.TimestampColumn.FillWidth"))
+		.DefaultLabel(LOCTEXT("Name.Label", "Time"))
+		.ToolTipText(LOCTEXT("Name.Description", "Indicates when this action was performed"));
+}
+
+TSharedRef<SWidget> FTimestampFileStateColumn::MakeColumnWidget(const FMakeFileStateColumnWidgetArgs& InArgs)
+{
+	return SNew(SBox)
+		.VAlign(VAlign_Center)
+		.Padding(FSandboxedEditingStyle::Get().GetMargin("SandboxedEditing.FileActions.TimestampColumn.Padding"))
+		[
+			SNew(STextBlock)
+			.Text_Lambda([Item = InArgs.RowData] { return TimestampDetail::GetRelativeTimestamp(*Item); })
+			.OverflowPolicy(ETextOverflowPolicy::MiddleEllipsis)
+			.HighlightText(InArgs.HighlightText)
+		];
+}
+}
+#undef LOCTEXT_NAMESPACE
\ No newline at end of file
diff --git a/Engine/Plugins/Developer/Sandbox/SandboxedEditing/Source/SandboxedEditing/Private/Features/Browser/Widgets/Shared/FileState/Columns/TimestampFileStateColumn.h b/Engine/Plugins/Developer/Sandbox/SandboxedEditing/Source/SandboxedEditing/Private/Features/Browser/Widgets/Shared/FileState/Columns/TimestampFileStateColumn.h
new file mode 100644
--- /dev/null
+++ b/Engine/Plugins/Developer/Sandbox/SandboxedEditing/Source/SandboxedEditing/Private/Features/Browser/Widgets/Shared/FileState/Columns/TimestampFileStateColumn.h
@@ -0,0 +1,25 @@
+﻿// Copyright Epic Games, Inc. All Rights Reserved.
+
+#pragma once
+
+#include "Features/Browser/ViewModels/FileState/IFileStateColumnBehavior.h"
+#include "Features/Browser/ViewModels/FileState/IFileStateColumnWidgetFactory.h"
+
+namespace UE::SandboxedEditing
+{
+/** Displays the timestamp of the sandboxed file change. */
+class FTimestampFileStateColumn : public IFileStateColumnBehavior, public IFileStateColumnWidgetFactory
+{
+public:
+	
+	//~ Begin IFileStateColumnBehavior Interface
+	virtual void PopulateSearchTerms(const TSharedPtr<FFileStateItem>& InRowData, TArray<FString>& OutSearchTerms) const override;
+	//~ End IFileStateColumnBehavior Interface
+
+	//~ Begin IFileStateWidgetFactory Interface
+	virtual SHeaderRow::FColumn::FArguments MakeColumnArguments() override;
+	virtual TSharedRef<SWidget> MakeColumnWidget(const FMakeFileStateColumnWidgetArgs& InArgs) override;
+	//~ End IFileStateWidgetFactory Interface
+};
+}
+
diff --git a/Engine/Plugins/Developer/Sandbox/SandboxedEditing/Source/SandboxedEditing/Private/Framework/Models/SandboxSystemModel.cpp b/Engine/Plugins/Developer/Sandbox/SandboxedEditing/Source/SandboxedEditing/Private/Framework/Models/SandboxSystemModel.cpp
--- a/Engine/Plugins/Developer/Sandbox/SandboxedEditing/Source/SandboxedEditing/Private/Framework/Models/SandboxSystemModel.cpp
+++ b/Engine/Plugins/Developer/Sandbox/SandboxedEditing/Source/SandboxedEditing/Private/Framework/Models/SandboxSystemModel.cpp
@@ -270,10 +270,11 @@ FileSandboxCore::FGatheredFileChanges FSandboxSystemModel::GatherFileChanges(con
 	ISandboxManager& SandboxManager = Module.GetSandboxManager();
 	
 	FGatheredFileChanges Result;
-	SandboxManager.EnumerateFileChanges(InSandbox, [&Result](const FString& InPath, ESandboxFileChange InChange)
+	SandboxManager.EnumerateFileChanges(InSandbox, [&InSandbox, &SandboxManager, &Result](const FString& InPath, ESandboxFileChange InChange)
 	{
 		Result.NonSandboxPaths.Add(InPath);
 		Result.FileActions.Add(InChange);
+		Result.Timestamps.Add(SandboxManager.GetSandboxedFileTimestamp(InSandbox, InPath).Get(FDateTime::MinValue()));
 		return FileSandboxCore::EBreakBehavior::Continue;
 	});
 	return Result;
@@ -290,6 +291,19 @@ void FSandboxSystemModel::EnumerateFileChanges(
 	SandboxManager.EnumerateFileChanges(InSandbox, InCallback);
 }
 
+FDateTime FSandboxSystemModel::GetSandboxedFileTimestamp(const FString& InPath) const
+{
+	using namespace FileSandboxCore;
+	IFileSandboxCoreModule& Module = IFileSandboxCoreModule::Get();
+	ISandboxManager& SandboxManager = Module.GetSandboxManager();
+	ISandboxInstance* Instance = SandboxManager.GetActiveSandboxInstance();
+	
+	const FDateTime MinValue = FDateTime::MinValue();
+	return Instance 
+		? Instance->GetSandboxedFileTimestamp(InPath).Get(MinValue) 
+		: MinValue;
+}
+
 void FSandboxSystemModel::DiscardAllChanges() const
 {
 	using namespace FileSandboxCore;
diff --git a/Engine/Plugins/Developer/Sandbox/SandboxedEditing/Source/SandboxedEditing/Private/Framework/Models/SandboxSystemModel.h b/Engine/Plugins/Developer/Sandbox/SandboxedEditing/Source/SandboxedEditing/Private/Framework/Models/SandboxSystemModel.h
--- a/Engine/Plugins/Developer/Sandbox/SandboxedEditing/Source/SandboxedEditing/Private/Framework/Models/SandboxSystemModel.h
+++ b/Engine/Plugins/Developer/Sandbox/SandboxedEditing/Source/SandboxedEditing/Private/Framework/Models/SandboxSystemModel.h
@@ -85,6 +85,9 @@ public:
 		const FString& InSandbox, 
 		TFunctionRef<FileSandboxCore::EBreakBehavior(const FString& InPath, FileSandboxCore::ESandboxFileChange InChange)> InCallback
 		);
+	
+	/** @return Timestamp of the sandboxed file action, or FDateTime::MinValue. */
+	FDateTime GetSandboxedFileTimestamp(const FString& InPath) const;
 
 	/** Discards all changes */
 	void DiscardAllChanges() const;
diff --git a/Engine/Plugins/Developer/Sandbox/SandboxedEditing/Source/SandboxedEditing/Private/SandboxedEditingStyle.cpp b/Engine/Plugins/Developer/Sandbox/SandboxedEditing/Source/SandboxedEditing/Private/SandboxedEditingStyle.cpp
--- a/Engine/Plugins/Developer/Sandbox/SandboxedEditing/Source/SandboxedEditing/Private/SandboxedEditingStyle.cpp
+++ b/Engine/Plugins/Developer/Sandbox/SandboxedEditing/Source/SandboxedEditing/Private/SandboxedEditingStyle.cpp
@@ -52,12 +52,14 @@ void FSandboxedEditingStyle::RegisterBrowserStyle()
 	Set("SandboxedEditing.Browser.LastModifiedColumn.FillWidth", 1.f);
 	
 	Set("SandboxedEditing.FileActions.RowPadding", FMargin(0.f, 4.f));
-	Set("SandboxedEditing.FileActions.PathColumn.Padding", FMargin(16.f, 0.f, 0.f, 0.f));
+	Set("SandboxedEditing.FileActions.PathColumn.Padding", FMargin(8.f, 0.f, 0.f, 0.f));
 	Set("SandboxedEditing.FileActions.PathColumn.FillWidth", 2.f);
 	Set("SandboxedEditing.FileActions.PersistCheckbox.Padding", FMargin(2.f, 0.f, 0.f, 0.f));
 	Set("SandboxedEditing.FileActions.PersistCheckbox.FixedWidth", 23.f);
 	Set("SandboxedEditing.FileActions.FileActionColumn.Padding", FMargin(0.f, 0.f, 0.f, 0.f));
 	Set("SandboxedEditing.FileActions.FileActionColumn.FixedWidth", 23.f);
+	Set("SandboxedEditing.FileActions.TimestampColumn.Padding", FMargin(8.f, 0.f, 0.f, 0.f));
+	Set("SandboxedEditing.FileActions.TimestampColumn.FillWidth", 200.f);
 }
 
 void FSandboxedEditingStyle::RegisterSharedStyle()
