Always use UtcNow() when getting the time to compare against.

#Synced-CL 49936699
#rnx

#Synced-CL 49936699

diff --git a/Engine/Plugins/Developer/Sandbox/FileSandbox/Source/FileSandboxCore/Private/Data/TimedAbsoluteFilePathArray.h b/Engine/Plugins/Developer/Sandbox/FileSandbox/Source/FileSandboxCore/Private/Data/TimedAbsoluteFilePathArray.h
--- a/Engine/Plugins/Developer/Sandbox/FileSandbox/Source/FileSandboxCore/Private/Data/TimedAbsoluteFilePathArray.h
+++ b/Engine/Plugins/Developer/Sandbox/FileSandbox/Source/FileSandboxCore/Private/Data/TimedAbsoluteFilePathArray.h
@@ -18,7 +18,7 @@ struct FFileSandboxCore_TimedAbsoluteFilePathArray
 	bool AreTimestampsValid() const { return Files.Num() == Timestamps.Num(); }
 
 	/** @return Whether the path was added, i.e. was not already present. */
-	bool Add(const FString& InPath, const FDateTime& InTimestamp = FDateTime::Now())
+	bool Add(const FString& InPath, const FDateTime& InTimestamp = FDateTime::UtcNow())
 	{
 		const int32 NumBefore = Files.Num();
 		const FString Path = FPaths::ConvertRelativePathToFull(InPath);
@@ -72,4 +72,4 @@ private:
 	 */
 	UPROPERTY()
 	TArray<FDateTime> Timestamps;
-};
\ No newline at end of file
+};
diff --git a/Engine/Plugins/Developer/Sandbox/SandboxedEditing/Source/SandboxedEditing/Private/Features/Browser/Widgets/Shared/FileState/Columns/TimestampFileStateColumn.cpp b/Engine/Plugins/Developer/Sandbox/SandboxedEditing/Source/SandboxedEditing/Private/Features/Browser/Widgets/Shared/FileState/Columns/TimestampFileStateColumn.cpp
--- a/Engine/Plugins/Developer/Sandbox/SandboxedEditing/Source/SandboxedEditing/Private/Features/Browser/Widgets/Shared/FileState/Columns/TimestampFileStateColumn.cpp
+++ b/Engine/Plugins/Developer/Sandbox/SandboxedEditing/Source/SandboxedEditing/Private/Features/Browser/Widgets/Shared/FileState/Columns/TimestampFileStateColumn.cpp
@@ -16,7 +16,9 @@ namespace TimestampDetail
 {
 static FText FormatRelativeTime(const FDateTime& EventTime, const FDateTime* CurrTime = nullptr)
 {
-	FTimespan TimeSpan = (CurrTime ? *CurrTime : FDateTime::Now()) - EventTime;
+	FDateTime BaseTime = CurrTime ? *CurrTime : FDateTime::UtcNow();
+
+	FTimespan TimeSpan = BaseTime - EventTime;
 	int32 Days = TimeSpan.GetDays();
 	int32 Hours = TimeSpan.GetHours();
 
@@ -87,4 +89,4 @@ TSharedRef<SWidget> FTimestampFileStateColumn::MakeColumnWidget(const FMakeFileS
 		];
 }
 }
-#undef LOCTEXT_NAMESPACE
\ No newline at end of file
+#undef LOCTEXT_NAMESPACE
