Disable sandbox creation when another tool has locked the sandbox

Added proper validation and UI feedback to prevent creating new sandboxes when
a leave lock is active (i.e., when another plugin has established a sandbox):

- Added CanCreateNewSandbox() method to check for leave lock before allowing
sandbox creation
- Hooked up CanExecuteAction to OpenCreateNewSandboxDialog command to disable
the button when locked
- Added tooltip override explaining why sandbox creation is disabled when locked
- Updated CanLoadSandbox() to also respect leave locks
- Refactored leave lock checking into reusable HasLeaveLock() helper method

This prevents users from attempting to create sandboxes that would fail, and
provides clear feedback about why the action is unavailable.
Prevent users from accessing the sandbox status bar when in a Multi-user session.

#Synced-CL 49936699

diff --git a/Engine/Plugins/Developer/Sandbox/SandboxedEditing/Source/SandboxedEditing/Private/Features/StatusBar/Commands/StatusBarCommandMappings.cpp b/Engine/Plugins/Developer/Sandbox/SandboxedEditing/Source/SandboxedEditing/Private/Features/StatusBar/Commands/StatusBarCommandMappings.cpp
--- a/Engine/Plugins/Developer/Sandbox/SandboxedEditing/Source/SandboxedEditing/Private/Features/StatusBar/Commands/StatusBarCommandMappings.cpp
+++ b/Engine/Plugins/Developer/Sandbox/SandboxedEditing/Source/SandboxedEditing/Private/Features/StatusBar/Commands/StatusBarCommandMappings.cpp
@@ -17,10 +17,11 @@ FStatusBarCommandMappings::FStatusBarCommandMappings(const TSharedRef<FSandboxSy
 	FStatusBarCommands::Register();
 	
 	FStatusBarCommands& Commands = FStatusBarCommands::Get();
-	
+
 	CommandList->MapAction(
 		Commands.OpenCreateNewSandboxDialog,
-		FExecuteAction::CreateStatic(&ShowNewSandboxModal, InSandboxViewModel.ToSharedPtr())
+		FExecuteAction::CreateStatic(&ShowNewSandboxModal, InSandboxViewModel.ToSharedPtr()),
+		FCanExecuteAction::CreateSP(InSandboxViewModel, &FSandboxSystemModel::CanCreateNewSandbox)
 		);
 	
 	CommandList->MapAction(
diff --git a/Engine/Plugins/Developer/Sandbox/SandboxedEditing/Source/SandboxedEditing/Private/Features/StatusBar/Widgets/SSandboxStatusBarWidget.cpp b/Engine/Plugins/Developer/Sandbox/SandboxedEditing/Source/SandboxedEditing/Private/Features/StatusBar/Widgets/SSandboxStatusBarWidget.cpp
--- a/Engine/Plugins/Developer/Sandbox/SandboxedEditing/Source/SandboxedEditing/Private/Features/StatusBar/Widgets/SSandboxStatusBarWidget.cpp
+++ b/Engine/Plugins/Developer/Sandbox/SandboxedEditing/Source/SandboxedEditing/Private/Features/StatusBar/Widgets/SSandboxStatusBarWidget.cpp
@@ -74,8 +74,17 @@ TSharedRef<SWidget> SSandboxStatusBarWidget::CreateStatusBarMenu()
 			FText Reason;
 			return Model->IsAllowedToLeaveSandbox(&Reason) ? FStatusBarCommands::Get().LeaveSandbox->GetDescription() : Reason;
 		});
+		const TAttribute<FText> OpenNewSandboxToolTipOverride = TAttribute<FText>::CreateLambda([Model = ViewModel]
+		{
+			if (!Model->CanCreateNewSandbox())
+			{
+				return LOCTEXT("SandboxEditing.CannotCreateSandbox",
+							   "Cannot create a new sandbox because some another tool is using the sandbox and has locked access.");
+			}
+			return LOCTEXT("SandboxEditing.CreateNewSandbox", "Create a new sandbox with dialog.");
+		});
 		
-		Section.AddMenuEntry(FStatusBarCommands::Get().OpenCreateNewSandboxDialog);
+		Section.AddMenuEntry(FStatusBarCommands::Get().OpenCreateNewSandboxDialog, {}, OpenNewSandboxToolTipOverride);
 		Section.AddMenuEntry(FStatusBarCommands::Get().LeaveSandbox, {}, LeaveToolTipOverride);
 		Section.AddMenuEntry(FStatusBarCommands::Get().PersistAll);
 		Section.AddMenuEntry(FStatusBarCommands::Get().DiscardAll);
@@ -132,4 +141,4 @@ void SSandboxStatusBarWidget::BuildSandboxesMenu(UToolMenu* InMenuBuilder)
 }
 }
 
-#undef LOCTEXT_NAMESPACE
\ No newline at end of file
+#undef LOCTEXT_NAMESPACE
diff --git a/Engine/Plugins/Developer/Sandbox/SandboxedEditing/Source/SandboxedEditing/Private/Framework/Models/SandboxSystemModel.cpp b/Engine/Plugins/Developer/Sandbox/SandboxedEditing/Source/SandboxedEditing/Private/Framework/Models/SandboxSystemModel.cpp
--- a/Engine/Plugins/Developer/Sandbox/SandboxedEditing/Source/SandboxedEditing/Private/Framework/Models/SandboxSystemModel.cpp
+++ b/Engine/Plugins/Developer/Sandbox/SandboxedEditing/Source/SandboxedEditing/Private/Framework/Models/SandboxSystemModel.cpp
@@ -84,6 +84,13 @@ bool FSandboxSystemModel::CreateNewSandbox(const FCreateSandboxArgs& InArgs)
 	return SandboxManager.CreateNewSandbox(Args).HasValue();
 }
 
+bool FSandboxSystemModel::CanCreateNewSandbox() const
+{
+	// If there is a leave lock in place it means that some other plugin has established a sandbox which would prevent us from
+	// creating a new one. 
+	return !HasLeaveLock();
+}
+	
 bool FSandboxSystemModel::CanCreateNewSandbox(const FString& InName, FText* Reason)
 {
 	FileSandboxCore::EDirectorySuitability Suitability = FileSandboxCore::DetermineDirectorySuitability(InName);
@@ -116,7 +123,7 @@ bool FSandboxSystemModel::LoadSandbox(const FString& InRootDirectory)
 
 bool FSandboxSystemModel::CanLoadSandbox(const FString& InRootDirectory)
 {
-	return !IsActiveSandbox(InRootDirectory);
+	return !IsActiveSandbox(InRootDirectory) && !HasLeaveLock();
 }
 
 bool FSandboxSystemModel::IsActiveSandbox(const FString& InRootDirectory) const
@@ -137,6 +144,17 @@ void FSandboxSystemModel::LeaveSandbox() const
 }
 
 #define SET_REASON(ReasonText) if (OutReason){ *OutReason = ReasonText; }
+
+bool FSandboxSystemModel::HasLeaveLock(FText* OutReason) const
+{
+	using namespace FileSandboxCore;
+	IFileSandboxCoreModule& Module = IFileSandboxCoreModule::Get();
+	ISandboxManager& SandboxManager = Module.GetSandboxManager();
+	
+	ISandboxLock* Lock = SandboxManager.GetActiveLock();
+	return Lock && !Lock->CanLeaveSandbox(OutReason);
+}
+	
 bool FSandboxSystemModel::IsAllowedToLeaveSandbox(FText* OutReason) const
 {
 	using namespace FileSandboxCore;
@@ -148,9 +166,8 @@ bool FSandboxSystemModel::IsAllowedToLeaveSandbox(FText* OutReason) const
 		SET_REASON(LOCTEXT("LeaveReason.NoSandbox", "No sandbox is active."));
 		return false;
 	}
-	
-	if (ISandboxLock* Lock = SandboxManager.GetActiveLock()
-		; Lock && !Lock->CanLeaveSandbox(OutReason))
+
+	if (HasLeaveLock(OutReason))
 	{
 		SET_REASON(LOCTEXT("LeaveReason.Locked", "The engine is locked to be sandboxed."));
 		return false;
@@ -459,4 +476,4 @@ void FSandboxSystemModel::HandleActiveSandboxFilesChanged()
 }
 }
 
-#undef LOCTEXT_NAMESPACE
\ No newline at end of file
+#undef LOCTEXT_NAMESPACE
diff --git a/Engine/Plugins/Developer/Sandbox/SandboxedEditing/Source/SandboxedEditing/Private/Framework/Models/SandboxSystemModel.h b/Engine/Plugins/Developer/Sandbox/SandboxedEditing/Source/SandboxedEditing/Private/Framework/Models/SandboxSystemModel.h
--- a/Engine/Plugins/Developer/Sandbox/SandboxedEditing/Source/SandboxedEditing/Private/Framework/Models/SandboxSystemModel.h
+++ b/Engine/Plugins/Developer/Sandbox/SandboxedEditing/Source/SandboxedEditing/Private/Framework/Models/SandboxSystemModel.h
@@ -40,7 +40,9 @@ public:
 	bool CreateNewSandbox(const FCreateSandboxArgs& InArgs);
 	/** @return Whether it is valid to call CreateNewSandbox. */
 	bool CanCreateNewSandbox(const FString& InName, FText* Reason = nullptr);
-
+	/** @return Whether we can create a new sandbox. This is allowed as long as there is no leave lock. */
+	bool CanCreateNewSandbox() const;
+		
 	/** Loads the specified sandbox. */
 	bool LoadSandbox(const FString& InRootDirectory);
 	/** @return Whether the user can join this sandbox (i.e. not yet in this sandbox) */
@@ -126,9 +128,11 @@ public:
 	
 	/** Invoked when the files in the active sandbox have changed. */
 	FSimpleMulticastDelegate& OnSandboxFilesChanged() { return OnSandboxFilesChangedDelegate; }
-		
+
 private:
-	
+	/** @return true if a leave lock is currently registered. */
+	bool HasLeaveLock(FText* OutReason = nullptr) const;
+		
 	/** Invoked when the result of GetKnownSandboxes may have changed. */
 	FSimpleMulticastDelegate OnKnownSandboxesChangedDelegate;
 	
