Adds a settings class for FileSandboxCore.  First setting exposed is a new setting that allows users to opt out of being asked to persist changes when leaving a sandbox.

#Synced-CL 50213209

#Synced-CL 50213209

diff --git a/Engine/Plugins/Developer/Sandbox/FileSandbox/Source/FileSandboxCore/Public/FileSandboxCoreSettings.h b/Engine/Plugins/Developer/Sandbox/FileSandbox/Source/FileSandboxCore/Public/FileSandboxCoreSettings.h
new file mode 100644
--- /dev/null
+++ b/Engine/Plugins/Developer/Sandbox/FileSandbox/Source/FileSandboxCore/Public/FileSandboxCoreSettings.h
@@ -0,0 +1,24 @@
+// Copyright Epic Games, Inc. All Rights Reserved.
+#pragma once
+
+#include "UObject/FieldPath.h"
+#include "UObject/UObjectGlobals.h"
+
+#include "FileSandboxCoreSettings.generated.h"
+
+#define UE_API SANDBOXEDEDITING_API
+
+UCLASS(MinimalAPI, config=Engine)
+class UFileSandboxCoreSettings : public UObject
+{
+	GENERATED_BODY()
+public:
+	/**
+	 * Should we always ask the user to persist sandbox contents when leaving a sandbox?
+	 */
+	UPROPERTY(config, EditAnywhere, Category = "Actions")
+	bool bAskToPersitWhenLeavingSandbox = true;
+};
+
+#undef UE_API
+
diff --git a/Engine/Plugins/Developer/Sandbox/SandboxedEditing/Source/SandboxedEditing/Private/Features/Browser/Views/AskUserToPersistView.cpp b/Engine/Plugins/Developer/Sandbox/SandboxedEditing/Source/SandboxedEditing/Private/Features/Browser/Views/AskUserToPersistView.cpp
--- a/Engine/Plugins/Developer/Sandbox/SandboxedEditing/Source/SandboxedEditing/Private/Features/Browser/Views/AskUserToPersistView.cpp
+++ b/Engine/Plugins/Developer/Sandbox/SandboxedEditing/Source/SandboxedEditing/Private/Features/Browser/Views/AskUserToPersistView.cpp
@@ -6,6 +6,7 @@
 #include "Internationalization/Internationalization.h"
 #include "Misc/MessageDialog.h"
 
+#include "FileSandboxCoreSettings.h"
 #define LOCTEXT_NAMESPACE "FAskUserToPersistView"
 
 namespace UE::SandboxedEditing
@@ -23,6 +24,14 @@ FAskUserToPersistView::~FAskUserToPersistView()
 
 void FAskUserToPersistView::ShowNotification(FAskUserToPersistWorkflow& InWorkflow) const
 {
+	UFileSandboxCoreSettings const* CurrentSettings = GetDefault<UFileSandboxCoreSettings>();
+	if (CurrentSettings && !CurrentSettings->bAskToPersitWhenLeavingSandbox)
+	{
+		// The user has indicated that they don't want to be prompted for saving a sandbox on existing the sandbox.
+		InWorkflow.LeaveWithoutPersist(); 
+		return;
+	}
+	
 	const EAppReturnType::Type Result = FMessageDialog::Open(
 		EAppMsgType::YesNoCancel, EAppReturnType::No,
 		LOCTEXT("Question", "You are about to leave a sandbox containing changes. Do you want to persist the changes?")
@@ -37,4 +46,4 @@ void FAskUserToPersistView::ShowNotification(FAskUserToPersistWorkflow& InWorkfl
 }
 }
 
-#undef LOCTEXT_NAMESPACE
\ No newline at end of file
+#undef LOCTEXT_NAMESPACE
diff --git a/Engine/Plugins/Developer/Sandbox/SandboxedEditing/Source/SandboxedEditing/Private/SandboxEditingModule.cpp b/Engine/Plugins/Developer/Sandbox/SandboxedEditing/Source/SandboxedEditing/Private/SandboxEditingModule.cpp
--- a/Engine/Plugins/Developer/Sandbox/SandboxedEditing/Source/SandboxedEditing/Private/SandboxEditingModule.cpp
+++ b/Engine/Plugins/Developer/Sandbox/SandboxedEditing/Source/SandboxedEditing/Private/SandboxEditingModule.cpp
@@ -6,9 +6,45 @@
 #include "Framework/Application/SlateApplication.h"
 #include "Misc/CoreDelegates.h"
 #include "SandboxedEditingStyle.h"
+#include "FileSandboxCoreSettings.h"
+
+#if WITH_EDITOR
+	#include "ISettingsModule.h"
+	#include "ISettingsSection.h"
+#endif	
+
+#define LOCTEXT_NAMESPACE "SandboxEditing"
 
 namespace UE::SandboxedEditing
 {
+void RegisterSettings()
+{
+#if WITH_EDITOR
+	ISettingsModule* SettingsModule = FModuleManager::GetModulePtr<ISettingsModule>("Settings");
+	if (SettingsModule != nullptr)
+	{
+		// TODO: make only one section for both settings objects
+		ISettingsSectionPtr SettingsSection = SettingsModule->RegisterSettings(
+			"Project", "Plugins", "Sandbox",
+			LOCTEXT("SandboxFileEditingSettings", "Sandbox Editing"),
+			LOCTEXT("SandboxFileEditingSettingsDescription", "Configure the sandboxed editing settings."),
+			GetMutableDefault<UFileSandboxCoreSettings>());
+	}
+#endif
+}
+void UnregisterSettings()
+{
+#if WITH_EDITOR
+	ISettingsModule* SettingsModule = FModuleManager::GetModulePtr<ISettingsModule>("Settings");
+	
+	if (SettingsModule != nullptr)
+	{
+		SettingsModule->UnregisterSettings("Project", "Plugins", "Sandbox");
+	}
+#endif
+}
+
+
 void FSandboxEditingModule::StartupModule()
 {
 	FSandboxedEditingStyle::Get();
@@ -18,14 +54,19 @@ void FSandboxEditingModule::StartupModule()
 		if (FSlateApplication::IsInitialized() && GIsEditor)
 		{
 			EditingApp = MakeUnique<FSandboxedEditingApp>();
+			RegisterSettings();
 		}
 	});
 }
 
+	
 void FSandboxEditingModule::ShutdownModule()
 {
 	EditingApp.Reset();
+	UnregisterSettings();
 }
-}
+} // namespace UE::SandboxedEditing
+
+#undef LOCTEXT_NAMESPACE
 
-IMPLEMENT_MODULE(UE::SandboxedEditing::FSandboxEditingModule, SandboxedEditing);
\ No newline at end of file
+IMPLEMENT_MODULE(UE::SandboxedEditing::FSandboxEditingModule, SandboxedEditing);
