Correct typo in the core sandbox settings.

#rnx

#Synced-CL 50213209

diff --git a/Engine/Plugins/Developer/Sandbox/FileSandbox/Source/FileSandboxCore/Public/FileSandboxCoreSettings.h b/Engine/Plugins/Developer/Sandbox/FileSandbox/Source/FileSandboxCore/Public/FileSandboxCoreSettings.h
--- a/Engine/Plugins/Developer/Sandbox/FileSandbox/Source/FileSandboxCore/Public/FileSandboxCoreSettings.h
+++ b/Engine/Plugins/Developer/Sandbox/FileSandbox/Source/FileSandboxCore/Public/FileSandboxCoreSettings.h
@@ -1,7 +1,6 @@
 // Copyright Epic Games, Inc. All Rights Reserved.
 #pragma once
 
-#include "UObject/FieldPath.h"
 #include "UObject/UObjectGlobals.h"
 
 #include "FileSandboxCoreSettings.generated.h"
@@ -16,8 +15,8 @@ public:
 	/**
 	 * Should we always ask the user to persist sandbox contents when leaving a sandbox?
 	 */
-	UPROPERTY(config, EditAnywhere, Category = "Actions")
-	bool bAskToPersitWhenLeavingSandbox = true;
+	UPROPERTY(config, EditAnywhere, Category = "Actions", meta = (DisplayName="Ask to persist when leaving the sandbox"))
+	bool bAskToPersistWhenLeavingSandbox = true;
 };
 
 #undef UE_API
diff --git a/Engine/Plugins/Developer/Sandbox/SandboxedEditing/Source/SandboxedEditing/Private/Features/Browser/Views/AskUserToPersistView.cpp b/Engine/Plugins/Developer/Sandbox/SandboxedEditing/Source/SandboxedEditing/Private/Features/Browser/Views/AskUserToPersistView.cpp
--- a/Engine/Plugins/Developer/Sandbox/SandboxedEditing/Source/SandboxedEditing/Private/Features/Browser/Views/AskUserToPersistView.cpp
+++ b/Engine/Plugins/Developer/Sandbox/SandboxedEditing/Source/SandboxedEditing/Private/Features/Browser/Views/AskUserToPersistView.cpp
@@ -25,7 +25,7 @@ FAskUserToPersistView::~FAskUserToPersistView()
 void FAskUserToPersistView::ShowNotification(FAskUserToPersistWorkflow& InWorkflow) const
 {
 	UFileSandboxCoreSettings const* CurrentSettings = GetDefault<UFileSandboxCoreSettings>();
-	if (CurrentSettings && !CurrentSettings->bAskToPersitWhenLeavingSandbox)
+	if (CurrentSettings && !CurrentSettings->bAskToPersistWhenLeavingSandbox)
 	{
 		// The user has indicated that they don't want to be prompted for saving a sandbox on existing the sandbox.
 		InWorkflow.LeaveWithoutPersist(); 
