Sandboxed Editing: Add right-click context menu action for browsing to the location that a changed file is saved at.

#Synced-CL 50904440

diff --git a/Engine/Plugins/Developer/Sandbox/FileSandbox/Source/FileSandboxCore/Private/Utils/SandboxDirectoryUtils.cpp b/Engine/Plugins/Developer/Sandbox/FileSandbox/Source/FileSandboxCore/Private/Utils/SandboxDirectoryUtils.cpp
--- a/Engine/Plugins/Developer/Sandbox/FileSandbox/Source/FileSandboxCore/Private/Utils/SandboxDirectoryUtils.cpp
+++ b/Engine/Plugins/Developer/Sandbox/FileSandbox/Source/FileSandboxCore/Private/Utils/SandboxDirectoryUtils.cpp
@@ -16,6 +16,8 @@
 #include "Misc/FileHelper.h"
 #include "Misc/Paths.h"
 #include "SandboxFileUtils.h"
+#include "Sandbox/Platform/SanboxedFilePathUtils.h"
+#include "Sandbox/Platform/SandboxedPlatformFilePath.h"
 #include "Templates/Function.h"
 #include "Types/EBreakBehavior.h"
 
@@ -265,6 +267,34 @@ FText FormatDirectorySuitabilityAsReason(EDirectorySuitability Suitability, cons
 	default: checkNoEntry(); return FText::GetEmpty();
 	}
 }
+
+TOptional<FString> GetSandboxPathFor(const FString& InSandboxDirectory, const FString& InNonSandboxPath)
+{
+	if (!IsRootSandboxDirectory(InSandboxDirectory))
+	{
+		return {};
+	}
+	
+	TOptional<FString> Result;
+	EnumerateMountPoints([&InSandboxDirectory, &InNonSandboxPath, &Result](const FString& InAssetPath, const FString& InFilesystemPath)
+	{
+		const FSandboxedPlatformFilePath MountPointPath = FSandboxedPlatformFilePath::CreateMountPoint(
+			*InSandboxDirectory, InAssetPath, InFilesystemPath
+			);
+		const TOptional<FSandboxedPlatformFilePath> Path = MakeSandboxedFilePathIfInMountPoint(InNonSandboxPath, MountPointPath);
+			
+		const bool bExists = Path && FPaths::FileExists(Path->GetNonSandboxPath());
+		if (bExists)
+		{
+			Result = Path->GetNonSandboxPath();
+			return EBreakBehavior::Break;
+		}
+			
+		return EBreakBehavior::Continue;
+	});
+	
+	return Result;
+}
 }
 
 #undef LOCTEXT_NAMESPACE
\ No newline at end of file
diff --git a/Engine/Plugins/Developer/Sandbox/FileSandbox/Source/FileSandboxCore/Public/Utils/SandboxDirectoryUtils.h b/Engine/Plugins/Developer/Sandbox/FileSandbox/Source/FileSandboxCore/Public/Utils/SandboxDirectoryUtils.h
--- a/Engine/Plugins/Developer/Sandbox/FileSandbox/Source/FileSandboxCore/Public/Utils/SandboxDirectoryUtils.h
+++ b/Engine/Plugins/Developer/Sandbox/FileSandbox/Source/FileSandboxCore/Public/Utils/SandboxDirectoryUtils.h
@@ -143,6 +143,16 @@ UE_API FText FormatDirectorySuitabilityAsReason(
 	const FString& InSandboxName,
 	const FString& InBaseDirectory = GetBaseSandboxDirectory()
 	);
+
+/** 
+ * Gets the file path to the file that saves the sandboxed state of InNonSandboxPath. 
+ * This only works for edited and added files. For removed files, this returns nothing.
+ * 
+ * @param InSandboxDirectory Path to the base directory of the sandbox, e.g. what's returned by ForEachSandbox. Either absolute or relative to the process BaseDir().
+ * @param InNonSandboxPath Path to a non-sandbox file of which you want to get the sandbox file.
+ * @return Path to sandboxed file corresponding to the non-sandbox file, if the file has a sandboxed edit or add operation. Unset, otherwise.
+ */
+UE_API TOptional<FString> GetSandboxPathFor(const FString& InSandboxDirectory, const FString& InNonSandboxPath);
 }
 
 #undef UE_API
\ No newline at end of file
diff --git a/Engine/Plugins/Developer/Sandbox/SandboxedEditing/Source/SandboxedEditing/Private/Features/Browser/Commands/BrowserCommandBindings.cpp b/Engine/Plugins/Developer/Sandbox/SandboxedEditing/Source/SandboxedEditing/Private/Features/Browser/Commands/BrowserCommandBindings.cpp
--- a/Engine/Plugins/Developer/Sandbox/SandboxedEditing/Source/SandboxedEditing/Private/Features/Browser/Commands/BrowserCommandBindings.cpp
+++ b/Engine/Plugins/Developer/Sandbox/SandboxedEditing/Source/SandboxedEditing/Private/Features/Browser/Commands/BrowserCommandBindings.cpp
@@ -16,8 +16,6 @@ FBrowserCommandBindings::FBrowserCommandBindings(const FBrowserViewModels& InVie
 	, PersistViewModel(InViewModels.PersistViewModel)
 	, CommandList(MakeShared<FUICommandList>())
 {
-	FBrowserCommands::Register();
-	
 	FBrowserCommands& BrowserCommands = FBrowserCommands::Get();
 	CommandList->MapAction(BrowserCommands.CreateNewSandbox, 
 		FExecuteAction::CreateSP(ControlsViewModel, &FSandboxControlsViewModel::StartCreationWorkflow), 
@@ -34,11 +32,6 @@ FBrowserCommandBindings::FBrowserCommandBindings(const FBrowserViewModels& InVie
 		);
 }
 
-FBrowserCommandBindings::~FBrowserCommandBindings()
-{
-	FBrowserCommands::Unregister();
-}
-
 void FBrowserCommandBindings::HandleCancel() const
 {
 	ControlsViewModel->CancelWorkflows();
diff --git a/Engine/Plugins/Developer/Sandbox/SandboxedEditing/Source/SandboxedEditing/Private/Features/Browser/Commands/BrowserCommandBindings.h b/Engine/Plugins/Developer/Sandbox/SandboxedEditing/Source/SandboxedEditing/Private/Features/Browser/Commands/BrowserCommandBindings.h
--- a/Engine/Plugins/Developer/Sandbox/SandboxedEditing/Source/SandboxedEditing/Private/Features/Browser/Commands/BrowserCommandBindings.h
+++ b/Engine/Plugins/Developer/Sandbox/SandboxedEditing/Source/SandboxedEditing/Private/Features/Browser/Commands/BrowserCommandBindings.h
@@ -20,7 +20,6 @@ class FBrowserCommandBindings : public FNoncopyable
 public:
 	
 	explicit FBrowserCommandBindings(const FBrowserViewModels& InViewModels);
-	~FBrowserCommandBindings();
 	
 	const TSharedRef<FUICommandList>& GetCommandList() const { return CommandList; }
 	
diff --git a/Engine/Plugins/Developer/Sandbox/SandboxedEditing/Source/SandboxedEditing/Private/Features/Browser/Commands/FileStateActions/FileStateCommandBindings.cpp b/Engine/Plugins/Developer/Sandbox/SandboxedEditing/Source/SandboxedEditing/Private/Features/Browser/Commands/FileStateActions/FileStateCommandBindings.cpp
new file mode 100644
--- /dev/null
+++ b/Engine/Plugins/Developer/Sandbox/SandboxedEditing/Source/SandboxedEditing/Private/Features/Browser/Commands/FileStateActions/FileStateCommandBindings.cpp
@@ -0,0 +1,91 @@
+﻿// Copyright Epic Games, Inc. All Rights Reserved.
+
+#include "FileStateCommandBindings.h"
+
+#include "FileStateCommands.h"
+#include "Features/Browser/Commands/BrowserCommands.h"
+#include "Features/Browser/ViewModels/FileState/FileStateItem.h"
+#include "Features/Browser/Views/Menu/FileStateMenuUtils.h"
+#include "Features/Browser/Widgets/Shared/FileState/SFileStateListView.h"
+#include "Framework/Models/SandboxSystemModel.h"
+#include "Framework/MultiBox/MultiBoxBuilder.h"
+
+namespace UE::SandboxedEditing
+{
+namespace CommonMenuDetail
+{
+static TArray<FString> TransformSelectedItems(TSharedPtr<SFileStateListView> InListView)
+{
+	const TArray<TSharedPtr<FFileStateItem>> SelectedItems = InListView->GetSelectedFiles();
+	TArray<FString> NonSandboxFiles;
+	Algo::Transform(SelectedItems, NonSandboxFiles, [](const TSharedPtr<FFileStateItem>& InItem)
+	{
+		return InItem->NonSandboxFile;
+	});
+	return NonSandboxFiles;
+}
+}
+
+FFileStateCommandBindings::FFileStateCommandBindings(TAttribute<TOptional<FString>> InSelectedSandboxRootAttr, TAttribute<TArray<FString>> InSelectedFilePathsAttr)
+	: SelectedSandboxRootAttr(MoveTemp(InSelectedSandboxRootAttr))
+	, SelectedFilePathsAttr(MoveTemp(InSelectedFilePathsAttr))
+{
+	BindCommands();
+}
+
+void FFileStateCommandBindings::AppendMenu(FMenuBuilder& InMenuBuilder)
+{
+	InMenuBuilder.PushCommandList(CommandList);
+	FileStateMenu::AppendMenu(InMenuBuilder);
+	InMenuBuilder.PopCommandList();
+}
+
+void FFileStateCommandBindings::BindCommands()
+{
+	FFileStateCommands Commands = FFileStateCommands::Get();
+	
+	CommandList->MapAction(
+		Commands.ShowRootInExplorer, FExecuteAction::CreateRaw(this, &FFileStateCommandBindings::HandleShowRootInExplorer)
+		);
+	CommandList->MapAction(
+		Commands.ShowFileInExplorer, 
+		FExecuteAction::CreateRaw(this, &FFileStateCommandBindings::HandleShowFileInExplorer),
+		FCanExecuteAction::CreateRaw(this, &FFileStateCommandBindings::CanShowFileInExplorer)
+		);
+}
+
+void FFileStateCommandBindings::HandleShowRootInExplorer() const
+{
+	if (const TOptional<FString> SandboxRoot = SelectedSandboxRootAttr.Get())
+	{
+		FPlatformProcess::ExploreFolder(**SandboxRoot);
+	}
+}
+
+void FFileStateCommandBindings::HandleShowFileInExplorer() const
+{
+	const TOptional<FString> SandboxRoot = SelectedSandboxRootAttr.Get();
+	const TArray<FString> SelectedPaths = SelectedFilePathsAttr.Get();
+	if (!SandboxRoot || SelectedPaths.IsEmpty())
+	{
+		return;
+	}
+	
+	if (const TOptional<FString> NonSandboxPath = SandboxModel::GetSandboxPathFor(*SandboxRoot, SelectedPaths[0]))
+	{
+		FPlatformProcess::ExploreFolder(**NonSandboxPath);
+	}
+}
+
+bool FFileStateCommandBindings::CanShowFileInExplorer() const
+{
+	const TOptional<FString> SandboxRoot = SelectedSandboxRootAttr.Get();
+	const TArray<FString> SelectedPaths = SelectedFilePathsAttr.Get();
+	if (!SandboxRoot || SelectedPaths.IsEmpty())
+	{
+		return false;
+	}
+	
+	return SandboxModel::GetSandboxPathFor(*SandboxRoot, SelectedPaths[0]).IsSet();
+}
+}
diff --git a/Engine/Plugins/Developer/Sandbox/SandboxedEditing/Source/SandboxedEditing/Private/Features/Browser/Commands/FileStateActions/FileStateCommandBindings.h b/Engine/Plugins/Developer/Sandbox/SandboxedEditing/Source/SandboxedEditing/Private/Features/Browser/Commands/FileStateActions/FileStateCommandBindings.h
new file mode 100644
--- /dev/null
+++ b/Engine/Plugins/Developer/Sandbox/SandboxedEditing/Source/SandboxedEditing/Private/Features/Browser/Commands/FileStateActions/FileStateCommandBindings.h
@@ -0,0 +1,53 @@
+﻿// Copyright Epic Games, Inc. All Rights Reserved.
+
+#pragma once
+
+#include "Framework/Commands/UICommandList.h"
+#include "Misc/Attribute.h"
+#include "Templates/SharedPointer.h"
+#include "Templates/UnrealTemplate.h"
+
+class FMenuBuilder;
+class FUICommandList;
+
+namespace UE::SandboxedEditing
+{
+class FSandboxSystemModel;
+class SFileStateListView;
+
+/** 
+ * Builds the right-click menu items displayed in all instances of SFileStateListView.
+ * Binds & unbinds the relevant commands required for that menu.
+ */
+class FFileStateCommandBindings : public FNoncopyable
+{
+public:
+	
+	explicit FFileStateCommandBindings(
+		TAttribute<TOptional<FString>> InSelectedSandboxRootAttr, TAttribute<TArray<FString>> InSelectedFilePathsAttr
+		);
+
+	/** Utility to append the commands to a menu builder. */
+	void AppendMenu(FMenuBuilder& InMenuBuilder);
+	
+private:
+	
+	/** The command list that our commands are bound to. */
+	const TSharedRef<FUICommandList> CommandList = MakeShared<FUICommandList>();
+	 	
+	/** The sandbox root path on which the commands should be run. */
+	const TAttribute<TOptional<FString>> SelectedSandboxRootAttr;
+	/** Gets the items selected in the UI. */
+	const TAttribute<TArray<FString>> SelectedFilePathsAttr;
+
+	/** Binds the commands needed for the menu. */
+	void BindCommands();
+	
+	void HandleShowRootInExplorer() const;
+	
+	void HandleShowFileInExplorer() const;
+	bool CanShowFileInExplorer() const;
+};
+}
+
+
diff --git a/Engine/Plugins/Developer/Sandbox/SandboxedEditing/Source/SandboxedEditing/Private/Features/Browser/Commands/FileStateActions/FileStateCommands.cpp b/Engine/Plugins/Developer/Sandbox/SandboxedEditing/Source/SandboxedEditing/Private/Features/Browser/Commands/FileStateActions/FileStateCommands.cpp
new file mode 100644
--- /dev/null
+++ b/Engine/Plugins/Developer/Sandbox/SandboxedEditing/Source/SandboxedEditing/Private/Features/Browser/Commands/FileStateActions/FileStateCommands.cpp
@@ -0,0 +1,27 @@
+﻿// Copyright Epic Games, Inc. All Rights Reserved.
+
+#include "FileStateCommands.h"
+
+#include "SandboxedEditingStyle.h"
+
+#define LOCTEXT_NAMESPACE "FFileStateCommands"
+
+namespace UE::SandboxedEditing
+{
+FFileStateCommands::FFileStateCommands()
+	: TCommands(
+		TEXT("FileState"), 
+		LOCTEXT("Description", "Commands relevant to the sandbox browser"), 
+		TEXT("Browser"), 
+		FSandboxedEditingStyle::Get().GetStyleSetName()
+	)
+{}
+
+void FFileStateCommands::RegisterCommands()
+{
+	UI_COMMAND(ShowRootInExplorer, "Show Root in Explorer", "Shows sandbox root directory in the system file explorer.", EUserInterfaceActionType::Button, FInputChord())
+	UI_COMMAND(ShowFileInExplorer, "Show File in Explorer", "Shows the file's sandboxed counterpart in the system file explorer.", EUserInterfaceActionType::Button, FInputChord())
+}
+}
+
+#undef LOCTEXT_NAMESPACE
diff --git a/Engine/Plugins/Developer/Sandbox/SandboxedEditing/Source/SandboxedEditing/Private/Features/Browser/Commands/FileStateActions/FileStateCommands.h b/Engine/Plugins/Developer/Sandbox/SandboxedEditing/Source/SandboxedEditing/Private/Features/Browser/Commands/FileStateActions/FileStateCommands.h
new file mode 100644
--- /dev/null
+++ b/Engine/Plugins/Developer/Sandbox/SandboxedEditing/Source/SandboxedEditing/Private/Features/Browser/Commands/FileStateActions/FileStateCommands.h
@@ -0,0 +1,27 @@
+﻿// Copyright Epic Games, Inc. All Rights Reserved.
+
+#pragma once
+
+#include "Framework/Commands/Commands.h"
+
+namespace UE::SandboxedEditing
+{
+/** Commands used for actions performed on sandboxed files (usually on selected items in SFileStateListView). */
+class FFileStateCommands : public TCommands<FFileStateCommands>
+{
+public:
+	
+	FFileStateCommands();
+	
+	//~ Begin TCommands Interface
+	virtual void RegisterCommands() override;
+	//~ End TCommands Interface
+	
+	/** Shows sandbox root directory in the system file explorer. */
+	TSharedPtr<FUICommandInfo> ShowRootInExplorer;
+	
+	/** Shows the file's sandboxed counterpart in the system file explorer. */
+	TSharedPtr<FUICommandInfo> ShowFileInExplorer;
+};
+}
+
diff --git a/Engine/Plugins/Developer/Sandbox/SandboxedEditing/Source/SandboxedEditing/Private/Features/Browser/Views/Menu/FileStateMenuUtils.cpp b/Engine/Plugins/Developer/Sandbox/SandboxedEditing/Source/SandboxedEditing/Private/Features/Browser/Views/Menu/FileStateMenuUtils.cpp
new file mode 100644
--- /dev/null
+++ b/Engine/Plugins/Developer/Sandbox/SandboxedEditing/Source/SandboxedEditing/Private/Features/Browser/Views/Menu/FileStateMenuUtils.cpp
@@ -0,0 +1,44 @@
+﻿// Copyright Epic Games, Inc. All Rights Reserved.
+
+#include "FileStateMenuUtils.h"
+
+#include "Containers/Array.h"
+#include "Containers/UnrealString.h"
+#include "Features/Browser/Commands/FileStateActions/FileStateCommands.h"
+#include "Features/Browser/ViewModels/FileState/FileStateItem.h"
+#include "Features/Browser/Widgets/Shared/FileState/SFileStateListView.h"
+#include "Framework/MultiBox/MultiBoxBuilder.h"
+
+namespace UE::SandboxedEditing::FileStateMenu
+{
+TArray<FString> TransformSelectedItems(const TSharedRef<SFileStateListView>& InListView)
+{
+	const TArray<TSharedPtr<FFileStateItem>> SelectedItems = InListView->GetSelectedFiles();
+	TArray<FString> NonSandboxFiles;
+	Algo::Transform(SelectedItems, NonSandboxFiles, [](const TSharedPtr<FFileStateItem>& InItem)
+	{
+		return InItem->NonSandboxFile;
+	});
+	return NonSandboxFiles;
+}
+
+TAttribute<TArray<FString>> FileStateMenu::MakeSelectedFilesAttribute(const TSharedRef<SFileStateListView>& InListView)
+{
+	return TAttribute<TArray<FString>>::CreateLambda([WeakListView = InListView.ToWeakPtr()]
+	{
+		const TSharedPtr<SFileStateListView> ListViewPin = WeakListView.Pin();
+		return ListViewPin ? TransformSelectedItems(ListViewPin.ToSharedRef()) : TArray<FString>{};
+	});
+}
+
+void FileStateMenu::AppendMenu(FMenuBuilder& InMenuBuilder)
+{
+	InMenuBuilder.BeginSection(TEXT("Common"));
+	
+	FFileStateCommands& Commands = FFileStateCommands::Get();
+	InMenuBuilder.AddMenuEntry(Commands.ShowRootInExplorer);
+	InMenuBuilder.AddMenuEntry(Commands.ShowFileInExplorer);
+	
+	InMenuBuilder.EndSection();
+}
+}
diff --git a/Engine/Plugins/Developer/Sandbox/SandboxedEditing/Source/SandboxedEditing/Private/Features/Browser/Views/Menu/FileStateMenuUtils.h b/Engine/Plugins/Developer/Sandbox/SandboxedEditing/Source/SandboxedEditing/Private/Features/Browser/Views/Menu/FileStateMenuUtils.h
new file mode 100644
--- /dev/null
+++ b/Engine/Plugins/Developer/Sandbox/SandboxedEditing/Source/SandboxedEditing/Private/Features/Browser/Views/Menu/FileStateMenuUtils.h
@@ -0,0 +1,22 @@
+﻿// Copyright Epic Games, Inc. All Rights Reserved.
+
+#pragma once
+
+#include "Containers/ContainersFwd.h"
+#include "Templates/SharedPointerFwd.h"
+
+class FMenuBuilder;
+class FUICommandList;
+namespace UE::SandboxedEditing { class SFileStateListView; }
+template<typename ObjectType> class TAttribute;
+
+namespace UE::SandboxedEditing::FileStateMenu
+{
+/** @return Gets the non-sandbox file paths selected in InListView. */
+TArray<FString> TransformSelectedItems(const TSharedRef<SFileStateListView>& InListView);
+/** @return Attribute returning result of TransformSelectedItems. */
+TAttribute<TArray<FString>> MakeSelectedFilesAttribute(const TSharedRef<SFileStateListView>& InListView);
+
+/** Appends entries that are shared across all menus. */
+void AppendMenu(FMenuBuilder& InMenuBuilder);
+}
diff --git a/Engine/Plugins/Developer/Sandbox/SandboxedEditing/Source/SandboxedEditing/Private/Features/Browser/Widgets/Sandboxed/SActiveSandbox.cpp b/Engine/Plugins/Developer/Sandbox/SandboxedEditing/Source/SandboxedEditing/Private/Features/Browser/Widgets/Sandboxed/SActiveSandbox.cpp
--- a/Engine/Plugins/Developer/Sandbox/SandboxedEditing/Source/SandboxedEditing/Private/Features/Browser/Widgets/Sandboxed/SActiveSandbox.cpp
+++ b/Engine/Plugins/Developer/Sandbox/SandboxedEditing/Source/SandboxedEditing/Private/Features/Browser/Widgets/Sandboxed/SActiveSandbox.cpp
@@ -2,11 +2,15 @@
 
 #include "SActiveSandbox.h"
 
+#include "SActiveSandboxFileChanges.h"
 #include "Features/Browser/ViewModels/Active/ActiveSandboxDetailsViewModel.h"
 #include "Features/Browser/ViewModels/FileState/Models/ActiveSandboxFileStateViewModel.h"
 #include "Features/Browser/Widgets/Shared/FileState/SFilterableFileStateListView.h"
 #include "Features/Browser/Widgets/Shared/SSandboxDescription.h"
 #include "SActiveSandboxToolbar.h"
+#include "Features/Browser/Commands/BrowserCommandBindings.h"
+#include "Features/Browser/Commands/BrowserCommands.h"
+#include "Framework/MultiBox/MultiBoxBuilder.h"
 #include "Styling/ToolBarStyle.h"
 #include "Widgets/Layout/SBorder.h"
 #include "Widgets/Layout/SBox.h"
@@ -14,6 +18,7 @@
 #include "Widgets/Layout/SSeparator.h"
 #include "Widgets/SBoxPanel.h"
 #include "Widgets/Text/STextBlock.h"
+#include "Widgets/Views/SListView.h"
 
 #define LOCTEXT_NAMESPACE "SActiveSandbox"
 
@@ -85,6 +90,10 @@ TSharedRef<SWidget> SActiveSandbox::MakeDescription(const FViewModels& InViewMod
 
 TSharedRef<SWidget> SActiveSandbox::MakeFileChangesContent(const FArguments& InArgs, const FViewModels& InViewModels)
 {
+	const SActiveSandboxFileChanges::FViewModels ViewModels(
+		InViewModels.ActiveSandboxViewModel, InViewModels.ActiveSandboxFileStateModel, InViewModels.FilterActiveSandboxFileStateModel
+		);
+	
 	return SNew(SExpandableArea)
 		.BodyBorderImage(FAppStyle::Get().GetBrush("ToolPanel.GroupBorder"))
 		.HeaderContent()
@@ -95,9 +104,9 @@ TSharedRef<SWidget> SActiveSandbox::MakeFileChangesContent(const FArguments& InA
 		]
 		.BodyContent()
 		[
-			SNew(SFilterableFileStateListView, InViewModels.ActiveSandboxFileStateModel, InViewModels.FilterActiveSandboxFileStateModel)
+			SNew(SActiveSandboxFileChanges, ViewModels)
+			.CommandList(InArgs._CommandList)
 			.ColumnFactories(InArgs._ColumnFactories)
-			.NoChangesText(LOCTEXT("NoItems", "File changes appear here once you edit, add, or remove a file."))
 		];
 }
 }
diff --git a/Engine/Plugins/Developer/Sandbox/SandboxedEditing/Source/SandboxedEditing/Private/Features/Browser/Widgets/Sandboxed/SActiveSandbox.h b/Engine/Plugins/Developer/Sandbox/SandboxedEditing/Source/SandboxedEditing/Private/Features/Browser/Widgets/Sandboxed/SActiveSandbox.h
--- a/Engine/Plugins/Developer/Sandbox/SandboxedEditing/Source/SandboxedEditing/Private/Features/Browser/Widgets/Sandboxed/SActiveSandbox.h
+++ b/Engine/Plugins/Developer/Sandbox/SandboxedEditing/Source/SandboxedEditing/Private/Features/Browser/Widgets/Sandboxed/SActiveSandbox.h
@@ -10,10 +10,11 @@ class FUICommandList;
 
 namespace UE::SandboxedEditing
 {
-class FFilterFileStateViewModel;
-class FActiveSandboxFileStateViewModel;
 class FActiveSandboxDetailsViewModel;
+class FActiveSandboxFileStateViewModel;
+class FFilterFileStateViewModel;
 class FSandboxMetaDataViewModel;
+class SFilterableFileStateListView;
 
 /** Widget displayed while in a sandbox */
 class SActiveSandbox : public SCompoundWidget
@@ -54,7 +55,7 @@ private:
 	
 	/** Details about the active sandbox. */
 	TSharedPtr<FActiveSandboxDetailsViewModel> ActiveSandboxViewModel;
-	
+
 	/** Make the description widget. */
 	TSharedRef<SWidget> MakeDescription(const FViewModels& InViewModels);
 	
diff --git a/Engine/Plugins/Developer/Sandbox/SandboxedEditing/Source/SandboxedEditing/Private/Features/Browser/Widgets/Sandboxed/SActiveSandboxFileChanges.cpp b/Engine/Plugins/Developer/Sandbox/SandboxedEditing/Source/SandboxedEditing/Private/Features/Browser/Widgets/Sandboxed/SActiveSandboxFileChanges.cpp
new file mode 100644
--- /dev/null
+++ b/Engine/Plugins/Developer/Sandbox/SandboxedEditing/Source/SandboxedEditing/Private/Features/Browser/Widgets/Sandboxed/SActiveSandboxFileChanges.cpp
@@ -0,0 +1,46 @@
+﻿// Copyright Epic Games, Inc. All Rights Reserved.
+
+#include "SActiveSandboxFileChanges.h"
+
+#include "Features/Browser/ViewModels/Active/ActiveSandboxDetailsViewModel.h"
+#include "Features/Browser/ViewModels/FileState/FileStateItem.h"
+#include "Features/Browser/Views/Menu/FileStateMenuUtils.h"
+#include "Features/Browser/Widgets/Shared/FileState/SFileStateListView.h"
+#include "Features/Browser/Widgets/Shared/FileState/SFilterableFileStateListView.h"
+#include "Framework/MultiBox/MultiBoxBuilder.h"
+
+#define LOCTEXT_NAMESPACE "SActiveSandboxFileChanges"
+
+namespace UE::SandboxedEditing
+{
+void SActiveSandboxFileChanges::Construct(
+	const FArguments& InArgs, const FViewModels& InViewModels
+	)
+{
+	ActiveSandboxViewModel = InViewModels.ActiveSandboxViewModel;
+	
+	ChildSlot
+	[
+		SAssignNew(FileActionWidget, SFilterableFileStateListView, InViewModels.FileActionsViewModel, InViewModels.FilterViewModel)
+		.ColumnFactories(InArgs._ColumnFactories)
+		.NoChangesText(LOCTEXT("NoItems", "File changes appear here once you edit, add, or remove a file."))
+		.OnContextMenuOpening(this, &SActiveSandboxFileChanges::MakeFileChangeContextMenu)	
+	];
+	
+	FileStateCommandBindings = MakeUnique<FFileStateCommandBindings>(
+		ActiveSandboxViewModel->GetSandboxPath(), 
+		FileStateMenu::MakeSelectedFilesAttribute(FileActionWidget->GetListView().ToSharedRef())
+		);
+}
+
+TSharedPtr<SWidget> SActiveSandboxFileChanges::MakeFileChangeContextMenu()
+{
+	FMenuBuilder Menu(true, nullptr);
+	
+	FileStateCommandBindings->AppendMenu(Menu);
+	
+	return Menu.MakeWidget();
+}
+}
+
+#undef LOCTEXT_NAMESPACE
\ No newline at end of file
diff --git a/Engine/Plugins/Developer/Sandbox/SandboxedEditing/Source/SandboxedEditing/Private/Features/Browser/Widgets/Sandboxed/SActiveSandboxFileChanges.h b/Engine/Plugins/Developer/Sandbox/SandboxedEditing/Source/SandboxedEditing/Private/Features/Browser/Widgets/Sandboxed/SActiveSandboxFileChanges.h
new file mode 100644
--- /dev/null
+++ b/Engine/Plugins/Developer/Sandbox/SandboxedEditing/Source/SandboxedEditing/Private/Features/Browser/Widgets/Sandboxed/SActiveSandboxFileChanges.h
@@ -0,0 +1,66 @@
+﻿// Copyright Epic Games, Inc. All Rights Reserved.
+
+#pragma once
+
+#include "Features/Browser/Commands/FileStateActions/FileStateCommandBindings.h"
+#include "Features/Browser/ViewModels/FileState/FileStateColumns.h"
+#include "Widgets/DeclarativeSyntaxSupport.h"
+#include "Widgets/SCompoundWidget.h"
+
+class FUICommandList;
+
+namespace UE::SandboxedEditing
+{
+class FActiveSandboxDetailsViewModel;
+class FFilterFileStateViewModel;
+class IFileStateViewModel;
+class SFilterableFileStateListView;
+
+/** Displays the file changes of the active sandbox. */
+class SActiveSandboxFileChanges : public SCompoundWidget
+{
+public:
+	
+	struct FViewModels
+	{
+		TSharedRef<FActiveSandboxDetailsViewModel> ActiveSandboxViewModel;
+		TSharedRef<IFileStateViewModel> FileActionsViewModel;
+		TSharedRef<FFilterFileStateViewModel> FilterViewModel;
+
+		explicit FViewModels(
+			const TSharedRef<FActiveSandboxDetailsViewModel>& ActiveSandboxViewModel,
+			const TSharedRef<IFileStateViewModel>& FileActionsViewModel,
+			const TSharedRef<FFilterFileStateViewModel>& FilterViewModel
+			)
+			: ActiveSandboxViewModel(ActiveSandboxViewModel)
+			, FileActionsViewModel(FileActionsViewModel)
+			, FilterViewModel(FilterViewModel)
+		{}
+	};
+	
+	SLATE_BEGIN_ARGS(SActiveSandboxFileChanges) {}
+		/** The command list to bind commands to. */
+		SLATE_ARGUMENT(TSharedPtr<FUICommandList>, CommandList)
+		
+		/** Used to build the columns for the file actions list. */
+		SLATE_ARGUMENT(FFileStateColumnFactoryMap, ColumnFactories)
+	SLATE_END_ARGS()
+	
+	void Construct(const FArguments& InArgs, const FViewModels& InViewModels);
+	
+private:
+	
+	/** Used for binding commands. */
+	TSharedPtr<FActiveSandboxDetailsViewModel> ActiveSandboxViewModel;
+	
+	/** Displays the file changes made in this sandbox. */
+	TSharedPtr<SFilterableFileStateListView> FileActionWidget;
+	
+	/** Common entries for FileActionWidget. */
+	TUniquePtr<FFileStateCommandBindings> FileStateCommandBindings;
+	
+	/** @return The context menu for FileActionWidget. */
+	TSharedPtr<SWidget> MakeFileChangeContextMenu();
+};
+}
+
diff --git a/Engine/Plugins/Developer/Sandbox/SandboxedEditing/Source/SandboxedEditing/Private/Features/Browser/Widgets/Shared/FileState/SFileStateListView.cpp b/Engine/Plugins/Developer/Sandbox/SandboxedEditing/Source/SandboxedEditing/Private/Features/Browser/Widgets/Shared/FileState/SFileStateListView.cpp
--- a/Engine/Plugins/Developer/Sandbox/SandboxedEditing/Source/SandboxedEditing/Private/Features/Browser/Widgets/Shared/FileState/SFileStateListView.cpp
+++ b/Engine/Plugins/Developer/Sandbox/SandboxedEditing/Source/SandboxedEditing/Private/Features/Browser/Widgets/Shared/FileState/SFileStateListView.cpp
@@ -32,8 +32,9 @@ void SFileStateListView::Construct(const FArguments& InArgs, const TSharedRef<IF
 			SAssignNew(ListView, SListView<TSharedPtr<FFileStateItem>>)
 			.ListItemsSource(&ViewModel->GetItems())
 			.OnGenerateRow(this, &SFileStateListView::OnGenerateRow)
-			.SelectionMode(ESelectionMode::None)
+			.SelectionMode(ESelectionMode::Multi)
 			.HeaderRow(MakeHeaderRow())
+			.OnContextMenuOpening(InArgs._OnContextMenuOpening)
 		]
 
 		+SOverlay::Slot()
@@ -53,6 +54,11 @@ void SFileStateListView::Construct(const FArguments& InArgs, const TSharedRef<IF
 	ViewModel->OnItemsChanged().AddSP(ListView.ToSharedRef(), &SListView<TSharedPtr<FFileStateItem>>::RequestListRefresh);
 }
 
+TArray<TSharedPtr<FFileStateItem>> SFileStateListView::GetSelectedFiles() const
+{
+	return ListView->GetSelectedItems();
+}
+
 TSharedRef<SHeaderRow> SFileStateListView::MakeHeaderRow()
 {
 	const TSharedRef<SHeaderRow> HeaderRow = SNew(SHeaderRow);
diff --git a/Engine/Plugins/Developer/Sandbox/SandboxedEditing/Source/SandboxedEditing/Private/Features/Browser/Widgets/Shared/FileState/SFileStateListView.h b/Engine/Plugins/Developer/Sandbox/SandboxedEditing/Source/SandboxedEditing/Private/Features/Browser/Widgets/Shared/FileState/SFileStateListView.h
--- a/Engine/Plugins/Developer/Sandbox/SandboxedEditing/Source/SandboxedEditing/Private/Features/Browser/Widgets/Shared/FileState/SFileStateListView.h
+++ b/Engine/Plugins/Developer/Sandbox/SandboxedEditing/Source/SandboxedEditing/Private/Features/Browser/Widgets/Shared/FileState/SFileStateListView.h
@@ -3,6 +3,7 @@
 #pragma once
 
 #include "Features/Browser/ViewModels/FileState/FileStateColumns.h"
+#include "Framework/Views/ITypedTableView.h"
 #include "Widgets/DeclarativeSyntaxSupport.h"
 #include "Widgets/SCompoundWidget.h"
 #include "Widgets/Views/SListView.h"
@@ -35,10 +36,16 @@ public:
 		 * Defaults to: "No file changes have been made in the sandbox."
 		 */
 		SLATE_ARGUMENT(FText, NoChangesText)
+
+		/** Used to generate the right-click context menu. */
+		SLATE_EVENT(FOnContextMenuOpening, OnContextMenuOpening)
 	SLATE_END_ARGS()
 
 	void Construct(const FArguments& InArgs, const TSharedRef<IFileStateViewModel>& InViewModel);
 	
+	/** @return Selected items. */
+	TArray<TSharedPtr<FFileStateItem>> GetSelectedFiles() const;
+	
 private:
 	
 	/** The view model being displayed */
diff --git a/Engine/Plugins/Developer/Sandbox/SandboxedEditing/Source/SandboxedEditing/Private/Features/Browser/Widgets/Shared/FileState/SFilterableFileStateListView.cpp b/Engine/Plugins/Developer/Sandbox/SandboxedEditing/Source/SandboxedEditing/Private/Features/Browser/Widgets/Shared/FileState/SFilterableFileStateListView.cpp
--- a/Engine/Plugins/Developer/Sandbox/SandboxedEditing/Source/SandboxedEditing/Private/Features/Browser/Widgets/Shared/FileState/SFilterableFileStateListView.cpp
+++ b/Engine/Plugins/Developer/Sandbox/SandboxedEditing/Source/SandboxedEditing/Private/Features/Browser/Widgets/Shared/FileState/SFilterableFileStateListView.cpp
@@ -38,11 +38,12 @@ void SFilterableFileStateListView::Construct(
 		[
 			SNew(SBorder)
 			[
-				SNew(SFileStateListView, InFileActionsViewModel)
+				SAssignNew(FileStateWidget, SFileStateListView, InFileActionsViewModel)
 				.ColumnFactories(InArgs._ColumnFactories)
 				.HighlightText_Lambda([InFilterViewModel]{ return InFilterViewModel->GetSearchText(); })
 				.OverrideOverlayText(this, &SFilterableFileStateListView::GetOverrideOverlayText)
 				.NoChangesText(InArgs._NoChangesText)
+				.OnContextMenuOpening(InArgs._OnContextMenuOpening)
 			]
 		]
 	];
diff --git a/Engine/Plugins/Developer/Sandbox/SandboxedEditing/Source/SandboxedEditing/Private/Features/Browser/Widgets/Shared/FileState/SFilterableFileStateListView.h b/Engine/Plugins/Developer/Sandbox/SandboxedEditing/Source/SandboxedEditing/Private/Features/Browser/Widgets/Shared/FileState/SFilterableFileStateListView.h
--- a/Engine/Plugins/Developer/Sandbox/SandboxedEditing/Source/SandboxedEditing/Private/Features/Browser/Widgets/Shared/FileState/SFilterableFileStateListView.h
+++ b/Engine/Plugins/Developer/Sandbox/SandboxedEditing/Source/SandboxedEditing/Private/Features/Browser/Widgets/Shared/FileState/SFilterableFileStateListView.h
@@ -3,25 +3,37 @@
 #pragma once
 
 #include "Features/Browser/ViewModels/FileState/FileStateColumns.h"
+#include "Framework/SlateDelegates.h"
+#include "Framework/Views/ITypedTableView.h"
 #include "Widgets/DeclarativeSyntaxSupport.h"
 #include "Widgets/SCompoundWidget.h"
 
+template <typename ItemType> class SListView;
+
 namespace UE::SandboxedEditing
 {
 class FFilterFileStateViewModel;
 class IFileStateViewModel;
+class SFileStateListView;
 
 /** Vertical arrangement of a filter section (search bar) and SFileStateListView. */
 class SFilterableFileStateListView : public SCompoundWidget
 {
 public:
 	
-	SLATE_BEGIN_ARGS(SFilterableFileStateListView){}
+	SLATE_BEGIN_ARGS(SFilterableFileStateListView)
+		{}
 		/** Used to build the columns. */
 		SLATE_ARGUMENT(FFileStateColumnFactoryMap, ColumnFactories)
-
+		
 		/** The text to display when there are no changes to display. */
 		SLATE_ARGUMENT(FText, NoChangesText)
+		
+		/** 
+		 * Used to generate the right-click context menu. 
+		 * @see GetListView
+		 */
+		SLATE_EVENT(FOnContextMenuOpening, OnContextMenuOpening)
 	SLATE_END_ARGS()
 
 	void Construct(
@@ -30,6 +42,9 @@ public:
 		const TSharedRef<FFilterFileStateViewModel>& InFilterViewModel
 		);
 	
+	/** @return List view displaying the items */
+	const TSharedPtr<SFileStateListView>& GetListView() const { return FileStateWidget; }
+	
 private:
 	
 	/** Used to generate the override overlay text. */
@@ -37,6 +52,9 @@ private:
 	/** Used to generate the override overlay text. */
 	TSharedPtr<FFilterFileStateViewModel> FilterViewModel;
 	
+	/** Displays the items. */
+	TSharedPtr<SFileStateListView> FileStateWidget;
+	
 	/** @return Text indicating that all items are filtered out, or empty. */
 	FText GetOverrideOverlayText() const;
 };
diff --git a/Engine/Plugins/Developer/Sandbox/SandboxedEditing/Source/SandboxedEditing/Private/Features/Browser/Widgets/Unsandboxed/SBrowserSandboxDetails.cpp b/Engine/Plugins/Developer/Sandbox/SandboxedEditing/Source/SandboxedEditing/Private/Features/Browser/Widgets/Unsandboxed/SBrowserSandboxDetails.cpp
--- a/Engine/Plugins/Developer/Sandbox/SandboxedEditing/Source/SandboxedEditing/Private/Features/Browser/Widgets/Unsandboxed/SBrowserSandboxDetails.cpp
+++ b/Engine/Plugins/Developer/Sandbox/SandboxedEditing/Source/SandboxedEditing/Private/Features/Browser/Widgets/Unsandboxed/SBrowserSandboxDetails.cpp
@@ -5,8 +5,10 @@
 #include "Components/VerticalBox.h"
 #include "Features/Browser/ViewModels/FileState/FileStateColumnRegistry.h"
 #include "Features/Browser/ViewModels/FileState/Models/UnloadedSandboxFileStateViewModel.h"
+#include "Features/Browser/Views/Menu/FileStateMenuUtils.h"
 #include "Features/Browser/Widgets/Shared/FileState/SFilterableFileStateListView.h"
 #include "Features/Browser/Widgets/Shared/SSandboxDescription.h"
+#include "Framework/MultiBox/MultiBoxBuilder.h"
 #include "Internationalization/Text.h"
 #include "Widgets/Layout/SBorder.h"
 #include "Widgets/Layout/SBox.h"
@@ -41,10 +43,16 @@ void SBrowserSandboxDetails::Construct(
 	)
 {
 	NumSelectedAttr = InArgs._NumSelected;
+	
 	ChildSlot
 	[
 		MakeWidgetSwitcherContent(InArgs, InMetaDataViewModel, InFileActionsViewModel, InFilterViewModel)
 	];
+	
+	FileStateCommandBindings = MakeUnique<FFileStateCommandBindings>(
+		InArgs._SandboxPath, 
+		FileStateMenu::MakeSelectedFilesAttribute(FileActionsView->GetListView().ToSharedRef())
+		);
 }
 
 TSharedRef<SWidget> SBrowserSandboxDetails::MakeWidgetSwitcherContent(
@@ -119,8 +127,9 @@ TSharedRef<SWidget> SBrowserSandboxDetails::MakeDetailContent(
 			]
 			.BodyContent()
 			[
-				SNew(SFilterableFileStateListView, InFileActionsViewModel, InFilterViewModel)
+				SAssignNew(FileActionsView, SFilterableFileStateListView, InFileActionsViewModel, InFilterViewModel)
                 .ColumnFactories(InArgs._ColumnFactories)
+				.OnContextMenuOpening(this, &SBrowserSandboxDetails::MakeFileChangeContextMenu)
 			]
 		];
 }
@@ -149,6 +158,15 @@ int32 SBrowserSandboxDetails::GetWidgetIndex() const
 	}
 	return TooManySelected;
 }
+
+TSharedPtr<SWidget> SBrowserSandboxDetails::MakeFileChangeContextMenu()
+{
+	FMenuBuilder MenuBuilder(true, nullptr);
+	
+	FileStateCommandBindings->AppendMenu(MenuBuilder);
+	
+	return MenuBuilder.MakeWidget();
+}
 }
 
 #undef LOCTEXT_NAMESPACE
\ No newline at end of file
diff --git a/Engine/Plugins/Developer/Sandbox/SandboxedEditing/Source/SandboxedEditing/Private/Features/Browser/Widgets/Unsandboxed/SBrowserSandboxDetails.h b/Engine/Plugins/Developer/Sandbox/SandboxedEditing/Source/SandboxedEditing/Private/Features/Browser/Widgets/Unsandboxed/SBrowserSandboxDetails.h
--- a/Engine/Plugins/Developer/Sandbox/SandboxedEditing/Source/SandboxedEditing/Private/Features/Browser/Widgets/Unsandboxed/SBrowserSandboxDetails.h
+++ b/Engine/Plugins/Developer/Sandbox/SandboxedEditing/Source/SandboxedEditing/Private/Features/Browser/Widgets/Unsandboxed/SBrowserSandboxDetails.h
@@ -2,12 +2,14 @@
 
 #pragma once
 
+#include "Features/Browser/Commands/FileStateActions/FileStateCommandBindings.h"
 #include "Features/Browser/ViewModels/FileState/FileStateColumns.h"
 #include "Widgets/DeclarativeSyntaxSupport.h"
 #include "Widgets/SCompoundWidget.h"
 
 namespace UE::SandboxedEditing
 {
+class SFilterableFileStateListView;
 class FFilterFileStateViewModel;
 class FSandboxMetaDataViewModel;
 class FUnloadedSandboxFileStateViewModel;
@@ -18,7 +20,7 @@ class SBrowserSandboxDetails : public SCompoundWidget
 public:
 	
 	SLATE_BEGIN_ARGS(SBrowserSandboxDetails){}
-		/** Path to the sandbox for which the description is being edited. Returning unset disables the edit box. */
+		/** Path to the selected sandbox. */
 		SLATE_ATTRIBUTE(TOptional<FString>, SandboxPath)
 		/** The number of selected sandboxes. Relevant for the content displayed. */
 		SLATE_ATTRIBUTE(int32, NumSelected)
@@ -38,6 +40,12 @@ private:
 	/** Gets the number of selected items. Determines the content displayed. */
 	TAttribute<int32> NumSelectedAttr;
 	
+	/** List of file actions being displayed. */
+	TSharedPtr<SFilterableFileStateListView> FileActionsView;
+	
+	/** Common entries for FileActionWidget. */
+	TUniquePtr<FFileStateCommandBindings> FileStateCommandBindings;
+	
 	/** @return The widget switcher that switches between the different messages and actual details content. */
 	TSharedRef<SWidget> MakeWidgetSwitcherContent(
 		const FArguments& InArgs, 
@@ -55,5 +63,7 @@ private:
 	
 	/** @return Index for widget switcher to display. Determined by how many items are selected. */
 	int32 GetWidgetIndex() const;
+
+	TSharedPtr<SWidget> MakeFileChangeContextMenu();
 };
 }
diff --git a/Engine/Plugins/Developer/Sandbox/SandboxedEditing/Source/SandboxedEditing/Private/Framework/Models/SandboxSystemModel.cpp b/Engine/Plugins/Developer/Sandbox/SandboxedEditing/Source/SandboxedEditing/Private/Framework/Models/SandboxSystemModel.cpp
--- a/Engine/Plugins/Developer/Sandbox/SandboxedEditing/Source/SandboxedEditing/Private/Framework/Models/SandboxSystemModel.cpp
+++ b/Engine/Plugins/Developer/Sandbox/SandboxedEditing/Source/SandboxedEditing/Private/Framework/Models/SandboxSystemModel.cpp
@@ -474,6 +474,14 @@ void FSandboxSystemModel::HandleActiveSandboxFilesChanged()
 {
 	OnSandboxFilesChangedDelegate.Broadcast();
 }
+
+namespace SandboxModel
+{
+TOptional<FString> SandboxModel::GetSandboxPathFor(const FString& InSandbox, const FString& InNonSandboxPath)
+{
+	return FileSandboxCore::GetSandboxPathFor(InSandbox, InNonSandboxPath);
+}
+}
 }
 
 #undef LOCTEXT_NAMESPACE
diff --git a/Engine/Plugins/Developer/Sandbox/SandboxedEditing/Source/SandboxedEditing/Private/Framework/Models/SandboxSystemModel.h b/Engine/Plugins/Developer/Sandbox/SandboxedEditing/Source/SandboxedEditing/Private/Framework/Models/SandboxSystemModel.h
--- a/Engine/Plugins/Developer/Sandbox/SandboxedEditing/Source/SandboxedEditing/Private/Framework/Models/SandboxSystemModel.h
+++ b/Engine/Plugins/Developer/Sandbox/SandboxedEditing/Source/SandboxedEditing/Private/Framework/Models/SandboxSystemModel.h
@@ -152,5 +152,15 @@ private:
 	/** Invoked when the set of files touched by the active sandbox changes. */
 	void HandleActiveSandboxFilesChanged();
 };
-}
 
+/** 
+ * Stateless functions that interact with the sandbox system.
+ * This is an abstraction to protect the UI code from the underlying FileSandboxCore types.
+ * TODO: Move stateless functions from FSandboxSystemModel into here.
+ */
+namespace SandboxModel
+{
+/** @return Gets the sandbox path for InNonSandboxPath. */
+TOptional<FString> GetSandboxPathFor(const FString& InSandbox, const FString& InNonSandboxPath);
+}
+}
\ No newline at end of file
diff --git a/Engine/Plugins/Developer/Sandbox/SandboxedEditing/Source/SandboxedEditing/Private/SandboxEditingModule.cpp b/Engine/Plugins/Developer/Sandbox/SandboxedEditing/Source/SandboxedEditing/Private/SandboxEditingModule.cpp
--- a/Engine/Plugins/Developer/Sandbox/SandboxedEditing/Source/SandboxedEditing/Private/SandboxEditingModule.cpp
+++ b/Engine/Plugins/Developer/Sandbox/SandboxedEditing/Source/SandboxedEditing/Private/SandboxEditingModule.cpp
@@ -7,6 +7,8 @@
 #include "Misc/CoreDelegates.h"
 #include "SandboxedEditingStyle.h"
 #include "FileSandboxCoreSettings.h"
+#include "Features/Browser/Commands/BrowserCommands.h"
+#include "Features/Browser/Commands/FileStateActions/FileStateCommands.h"
 
 #if WITH_EDITOR
 	#include "ISettingsModule.h"
@@ -49,6 +51,9 @@ void FSandboxEditingModule::StartupModule()
 {
 	FSandboxedEditingStyle::Get();
 	
+	FBrowserCommands::Register();
+	FFileStateCommands::Register();
+	
 	FCoreDelegates::OnPostEngineInit.AddLambda([this]
 	{
 		if (FSlateApplication::IsInitialized() && GIsEditor)
@@ -64,6 +69,9 @@ void FSandboxEditingModule::ShutdownModule()
 {
 	EditingApp.Reset();
 	UnregisterSettings();
+	
+	FBrowserCommands::Unregister();
+	FFileStateCommands::Unregister();
 }
 } // namespace UE::SandboxedEditing
 
