Sandboxed Editing: Move UFileSandboxCoreSettings to Sandboxed Editing module and rename to USandboxedEditingSettings.

#Synced-CL 50904440

diff --git a/Engine/Plugins/Developer/Sandbox/FileSandbox/Source/FileSandboxCore/Public/FileSandboxCoreSettings.h b/Engine/Plugins/Developer/Sandbox/FileSandbox/Source/FileSandboxCore/Public/FileSandboxCoreSettings.h
deleted file mode 100644
diff --git a/Engine/Plugins/Developer/Sandbox/SandboxedEditing/Source/SandboxedEditing/Private/Features/Browser/Views/AskUserToPersistView.cpp b/Engine/Plugins/Developer/Sandbox/SandboxedEditing/Source/SandboxedEditing/Private/Features/Browser/Views/AskUserToPersistView.cpp
--- a/Engine/Plugins/Developer/Sandbox/SandboxedEditing/Source/SandboxedEditing/Private/Features/Browser/Views/AskUserToPersistView.cpp
+++ b/Engine/Plugins/Developer/Sandbox/SandboxedEditing/Source/SandboxedEditing/Private/Features/Browser/Views/AskUserToPersistView.cpp
@@ -2,11 +2,11 @@
 
 #include "AskUserToPersistView.h"
 
+#include "SandboxedEditingSettings.h"
 #include "Features/Browser/ViewModels/Leaving/LeaveSandboxViewModel.h"
 #include "Internationalization/Internationalization.h"
 #include "Misc/MessageDialog.h"
 
-#include "FileSandboxCoreSettings.h"
 #define LOCTEXT_NAMESPACE "FAskUserToPersistView"
 
 namespace UE::SandboxedEditing
@@ -24,7 +24,7 @@ FAskUserToPersistView::~FAskUserToPersistView()
 
 void FAskUserToPersistView::ShowNotification(FAskUserToPersistWorkflow& InWorkflow) const
 {
-	UFileSandboxCoreSettings const* CurrentSettings = GetDefault<UFileSandboxCoreSettings>();
+	USandboxedEditingSettings const* CurrentSettings = USandboxedEditingSettings::Get();
 	if (CurrentSettings && !CurrentSettings->bAskToPersistWhenLeavingSandbox)
 	{
 		// The user has indicated that they don't want to be prompted for saving a sandbox on existing the sandbox.
diff --git a/Engine/Plugins/Developer/Sandbox/SandboxedEditing/Source/SandboxedEditing/Private/SandboxEditingModule.cpp b/Engine/Plugins/Developer/Sandbox/SandboxedEditing/Source/SandboxedEditing/Private/SandboxEditingModule.cpp
--- a/Engine/Plugins/Developer/Sandbox/SandboxedEditing/Source/SandboxedEditing/Private/SandboxEditingModule.cpp
+++ b/Engine/Plugins/Developer/Sandbox/SandboxedEditing/Source/SandboxedEditing/Private/SandboxEditingModule.cpp
@@ -3,17 +3,14 @@
 #include "SandboxEditingModule.h"
 
 #include "CoreGlobals.h"
+#include "Features/Browser/Commands/BrowserCommands.h"
+#include "Features/Browser/Commands/FileStateActions/FileStateCommands.h"
 #include "Framework/Application/SlateApplication.h"
+#include "ISettingsModule.h"
+#include "ISettingsSection.h"
 #include "Misc/CoreDelegates.h"
+#include "SandboxedEditingSettings.h"
 #include "SandboxedEditingStyle.h"
-#include "FileSandboxCoreSettings.h"
-#include "Features/Browser/Commands/BrowserCommands.h"
-#include "Features/Browser/Commands/FileStateActions/FileStateCommands.h"
-
-#if WITH_EDITOR
-	#include "ISettingsModule.h"
-	#include "ISettingsSection.h"
-#endif	
 
 #define LOCTEXT_NAMESPACE "SandboxEditing"
 
@@ -21,32 +18,24 @@ namespace UE::SandboxedEditing
 {
 void RegisterSettings()
 {
-#if WITH_EDITOR
 	ISettingsModule* SettingsModule = FModuleManager::GetModulePtr<ISettingsModule>("Settings");
 	if (SettingsModule != nullptr)
 	{
-		// TODO: make only one section for both settings objects
 		ISettingsSectionPtr SettingsSection = SettingsModule->RegisterSettings(
 			"Project", "Plugins", "Sandbox",
 			LOCTEXT("SandboxFileEditingSettings", "Sandbox Editing"),
 			LOCTEXT("SandboxFileEditingSettingsDescription", "Configure the sandboxed editing settings."),
-			GetMutableDefault<UFileSandboxCoreSettings>());
+			USandboxedEditingSettings::Get());
 	}
-#endif
 }
 void UnregisterSettings()
 {
-#if WITH_EDITOR
-	ISettingsModule* SettingsModule = FModuleManager::GetModulePtr<ISettingsModule>("Settings");
-	
-	if (SettingsModule != nullptr)
+	if (ISettingsModule* SettingsModule = FModuleManager::GetModulePtr<ISettingsModule>("Settings"))
 	{
 		SettingsModule->UnregisterSettings("Project", "Plugins", "Sandbox");
 	}
-#endif
 }
 
-
 void FSandboxEditingModule::StartupModule()
 {
 	FSandboxedEditingStyle::Get();
diff --git a/Engine/Plugins/Developer/Sandbox/SandboxedEditing/Source/SandboxedEditing/Private/SandboxedEditingSettings.h b/Engine/Plugins/Developer/Sandbox/SandboxedEditing/Source/SandboxedEditing/Private/SandboxedEditingSettings.h
new file mode 100644
--- /dev/null
+++ b/Engine/Plugins/Developer/Sandbox/SandboxedEditing/Source/SandboxedEditing/Private/SandboxedEditingSettings.h
@@ -0,0 +1,21 @@
+// Copyright Epic Games, Inc. All Rights Reserved.
+
+#pragma once
+
+#include "UObject/UObjectGlobals.h"
+#include "SandboxedEditingSettings.generated.h"
+
+UCLASS(MinimalAPI, config=Engine)
+class USandboxedEditingSettings : public UObject
+{
+	GENERATED_BODY()
+public:
+	
+	static USandboxedEditingSettings* Get() { return GetMutableDefault<USandboxedEditingSettings>(); }
+	
+	/** Should we always ask the user to persist sandbox contents when leaving a sandbox? */
+	UPROPERTY(config, EditAnywhere, Category = "Actions", meta = (DisplayName="Ask to persist when leaving the sandbox"))
+	bool bAskToPersistWhenLeavingSandbox = true;
+};
+
+
