Sandboxed Editing: Implement exporting and importing of sandboxes.
- The browser toolbar now has export and import buttons
- The export button exports as many items as are selected in the browser as .zip.
- The zip export path is saved in global settings so the editor remembers the same export path for the next time

#Synced-CL 50904440

diff --git a/Engine/Plugins/Developer/Sandbox/SandboxedEditing/Source/SandboxedEditing/Private/Features/Browser/Commands/BrowserCommands.cpp b/Engine/Plugins/Developer/Sandbox/SandboxedEditing/Source/SandboxedEditing/Private/Features/Browser/Commands/BrowserCommands.cpp
--- a/Engine/Plugins/Developer/Sandbox/SandboxedEditing/Source/SandboxedEditing/Private/Features/Browser/Commands/BrowserCommands.cpp
+++ b/Engine/Plugins/Developer/Sandbox/SandboxedEditing/Source/SandboxedEditing/Private/Features/Browser/Commands/BrowserCommands.cpp
@@ -24,6 +24,9 @@ void FBrowserCommands::RegisterCommands()
 	
 	UI_COMMAND(LeaveSandbox, "Leave Sandbox", "Leaves the current sandbox.", EUserInterfaceActionType::Button, FInputChord());
 	UI_COMMAND(PersistSandbox, "Persist Sandbox", "Persists the active sandbox.", EUserInterfaceActionType::Button, FInputChord());
+	
+	UI_COMMAND(ExportSandboxes, "Export", "Exports the sandboxes selected in the browser as zip archive.", EUserInterfaceActionType::Button, FInputChord());
+	UI_COMMAND(ImportSandboxes, "Import", "Imports sandboxes saved in zip archives.", EUserInterfaceActionType::Button, FInputChord());
 }
 }
 
diff --git a/Engine/Plugins/Developer/Sandbox/SandboxedEditing/Source/SandboxedEditing/Private/Features/Browser/Commands/BrowserCommands.h b/Engine/Plugins/Developer/Sandbox/SandboxedEditing/Source/SandboxedEditing/Private/Features/Browser/Commands/BrowserCommands.h
--- a/Engine/Plugins/Developer/Sandbox/SandboxedEditing/Source/SandboxedEditing/Private/Features/Browser/Commands/BrowserCommands.h
+++ b/Engine/Plugins/Developer/Sandbox/SandboxedEditing/Source/SandboxedEditing/Private/Features/Browser/Commands/BrowserCommands.h
@@ -27,6 +27,11 @@ public:
 	TSharedPtr<FUICommandInfo> LeaveSandbox;
 	/** Persists the current sandbox (without leaving). */
 	TSharedPtr<FUICommandInfo> PersistSandbox;
+	
+	/** Exports the sandboxes selected in the browser. */
+	TSharedPtr<FUICommandInfo> ExportSandboxes;
+	/** Imports sandboxes selected from the system explorer. */
+	TSharedPtr<FUICommandInfo> ImportSandboxes;
 };
 }
 
diff --git a/Engine/Plugins/Developer/Sandbox/SandboxedEditing/Source/SandboxedEditing/Private/Features/Browser/ViewModels/SandboxControlsViewModel.cpp b/Engine/Plugins/Developer/Sandbox/SandboxedEditing/Source/SandboxedEditing/Private/Features/Browser/ViewModels/SandboxControlsViewModel.cpp
--- a/Engine/Plugins/Developer/Sandbox/SandboxedEditing/Source/SandboxedEditing/Private/Features/Browser/ViewModels/SandboxControlsViewModel.cpp
+++ b/Engine/Plugins/Developer/Sandbox/SandboxedEditing/Source/SandboxedEditing/Private/Features/Browser/ViewModels/SandboxControlsViewModel.cpp
@@ -86,6 +86,39 @@ void FSandboxControlsViewModel::StartRenameWorkflow(const FString& InSandboxRoot
 	}
 }
 
+bool FSandboxControlsViewModel::CanStartExportWorkflow() const
+{
+	return !HasActiveWorkflow();
+}
+
+void FSandboxControlsViewModel::StartExportWorkflow(const TArray<FString>& InSandboxRoots)
+{
+	if (CanStartExportWorkflow() && !InSandboxRoots.IsEmpty())
+	{
+		CurrentExportWorkflow = MakeUnique<FExportWorkflow>(
+			InSandboxRoots,
+			FSimpleDelegate::CreateRaw(this, &FSandboxControlsViewModel::HandleExportWorkflowEnded)
+			);
+		OnExportWorkflowStartedDelegate.Broadcast(*CurrentExportWorkflow);
+	}
+}
+
+bool FSandboxControlsViewModel::CanStartImportWorkflow() const
+{
+	return !HasActiveWorkflow();
+}
+
+void FSandboxControlsViewModel::StartImportWorkflow()
+{
+	if (CanStartImportWorkflow())
+	{
+		CurrentImportWorkflow = MakeUnique<FImportWorkflow>(
+			FSimpleDelegate::CreateRaw(this, &FSandboxControlsViewModel::HandleImportWorkflowEnded)
+			);
+		OnImportWorkflowStartedDelegate.Broadcast(*CurrentImportWorkflow);
+	}
+}
+
 void FSandboxControlsViewModel::HandleCreationWorkflowEnded()
 {
 	CurrentCreationWorkflow.Reset();
@@ -97,4 +130,16 @@ void FSandboxControlsViewModel::HandleRenameWorkflowEnded()
 	CurrentRenameWorkflow.Reset();
 	OnRenameWorkflowEndedDelegate.Broadcast();
 }
+
+void FSandboxControlsViewModel::HandleExportWorkflowEnded()
+{
+	CurrentExportWorkflow.Reset();
+	OnExportWorkflowEndedDelegate.Broadcast();
+}
+
+void FSandboxControlsViewModel::HandleImportWorkflowEnded()
+{
+	CurrentImportWorkflow.Reset();
+	OnImportWorkflowEndedDelegate.Broadcast();
+}
 }
diff --git a/Engine/Plugins/Developer/Sandbox/SandboxedEditing/Source/SandboxedEditing/Private/Features/Browser/ViewModels/SandboxControlsViewModel.h b/Engine/Plugins/Developer/Sandbox/SandboxedEditing/Source/SandboxedEditing/Private/Features/Browser/ViewModels/SandboxControlsViewModel.h
--- a/Engine/Plugins/Developer/Sandbox/SandboxedEditing/Source/SandboxedEditing/Private/Features/Browser/ViewModels/SandboxControlsViewModel.h
+++ b/Engine/Plugins/Developer/Sandbox/SandboxedEditing/Source/SandboxedEditing/Private/Features/Browser/ViewModels/SandboxControlsViewModel.h
@@ -8,11 +8,15 @@
 #include "SandboxCreationWorkflow.h"
 #include "Templates/UniquePtr.h"
 #include "Templates/UnrealTemplate.h"
+#include "Transfer/ExportWorkflow.h"
+#include "Transfer/ImportWorkflow.h"
 
 namespace UE::SandboxedEditing
 {
 DECLARE_MULTICAST_DELEGATE_OneParam(FCreationWorkflowDelegate, FSandboxCreationWorkflow&);
 DECLARE_MULTICAST_DELEGATE_OneParam(FRenameWorkflowDelegate, FRenameSandboxWorkflow&);
+DECLARE_MULTICAST_DELEGATE_OneParam(FExportWorkflowDelegate, FExportWorkflow&);
+DECLARE_MULTICAST_DELEGATE_OneParam(FImportWorkflowDelegate, FImportWorkflow&);
 
 /**
  * Handles actions that can be performed with FSandboxListItems, such as deleting, renaming, loading.
@@ -53,6 +57,20 @@ public:
 	/** @return The currently active rename workflow. */
 	FRenameSandboxWorkflow* GetCurrentRenameWorkflow() const { return CurrentRenameWorkflow.Get(); }
 	
+	/** @return Whether an export operation can be performed. */
+	bool CanStartExportWorkflow() const;
+	/** Starts an export workflow */
+	void StartExportWorkflow(const TArray<FString>& InSandboxRoots);
+	/** @return The currently active export workflow. */
+	FExportWorkflow* GetCurrentExportWorkflow() const { return CurrentExportWorkflow.Get(); }
+	
+	/** @return Whether an import operation can be performed. */
+	bool CanStartImportWorkflow() const;
+	/** Starts an import workflow. */
+	void StartImportWorkflow();
+	/** @return The currently active import workflow. */
+	FImportWorkflow* GetCurrentImportWorkflow() const { return CurrentImportWorkflow.Get(); }
+	
 	/** @return Delegate that is invoked when the creation workflow starts. */
 	FCreationWorkflowDelegate& OnCreationWorkflowStarted() { return OnCreationWorkflowStartedDelegate; }
 	/** @return Delegate that is invoked when the creation workflow ends. */
@@ -63,6 +81,16 @@ public:
 	/** @return Delegate that is invoked when the rename workflow ends. */
 	FSimpleMulticastDelegate& OnRenameWorkflowEnded() { return OnRenameWorkflowEndedDelegate; }
 	
+	/** @return Delegate that is invoked when the export workflow starts. */
+	FExportWorkflowDelegate& OnExportWorkflowStarted() { return OnExportWorkflowStartedDelegate; }
+	/** @return Delegate that is invoked when the export workflow ends. */
+	FSimpleMulticastDelegate& OnExportWorkflowEnded() { return OnExportWorkflowEndedDelegate; }
+	
+	/** @return Delegate that is invoked when the import workflow starts. */
+	FImportWorkflowDelegate& OnImportWorkflowStarted() { return OnImportWorkflowStartedDelegate; }
+	/** @return Delegate that is invoked when the import workflow ends. */
+	FSimpleMulticastDelegate& OnImportWorkflowEnded() { return OnImportWorkflowEndedDelegate; }
+	
 private:
 	
 	/** The underlying model that performs the mutations. */
@@ -72,24 +100,42 @@ private:
 	TUniquePtr<FSandboxCreationWorkflow> CurrentCreationWorkflow;
 	/** Set if the user is currently in the process of renaming a sandbox. */
 	TUniquePtr<FRenameSandboxWorkflow> CurrentRenameWorkflow;
+	/** Set if the user is currently in the process of exporting sandboxes. */
+	TUniquePtr<FExportWorkflow> CurrentExportWorkflow;
+	/** Set if the user is currently in the process of importing sandboxes. */
+	TUniquePtr<FImportWorkflow> CurrentImportWorkflow;
 	
 	/** Invoked when the creation workflow starts. */
 	FCreationWorkflowDelegate OnCreationWorkflowStartedDelegate;
 	/** Invoked when the creation workflow ends. */
 	FSimpleMulticastDelegate OnCreationWorkflowEndedDelegate;
 	
-	/** Invoked when the creation rename starts. */
+	/** Invoked when the rename workflow starts. */
 	FRenameWorkflowDelegate OnRenameWorkflowStartedDelegate;
-	/** Invoked when the creation rename ends. */
+	/** Invoked when the rename workflow ends. */
 	FSimpleMulticastDelegate OnRenameWorkflowEndedDelegate;
 	
+	/** Invoked when the export workflow starts. */
+	FExportWorkflowDelegate OnExportWorkflowStartedDelegate;
+	/** Invoked when the export workflow ends. */
+	FSimpleMulticastDelegate OnExportWorkflowEndedDelegate;
+	
+	/** Invoked when the import workflow starts. */
+	FImportWorkflowDelegate OnImportWorkflowStartedDelegate;
+	/** Invoked when the import workflow ends. */
+	FSimpleMulticastDelegate OnImportWorkflowEndedDelegate;
+	
 	/** Called by CurrentCreationWorkflow when the workflow ends. */
 	void HandleCreationWorkflowEnded();
 	/** Called by CurrentRenameWorkflow when the workflow ends. */
 	void HandleRenameWorkflowEnded();
+	/** Called by CurrentExportWorkflow when the workflow ends. */
+	void HandleExportWorkflowEnded();
+	/** Called by CurrentImportWorkflow when the workflow ends. */
+	void HandleImportWorkflowEnded();
 	
 	/** @return Whether a workflow is active */
-	bool HasActiveWorkflow() const { return CurrentCreationWorkflow || CurrentRenameWorkflow; }
+	bool HasActiveWorkflow() const { return CurrentCreationWorkflow || CurrentRenameWorkflow || CurrentExportWorkflow || CurrentImportWorkflow; }
 };
 }
 
diff --git a/Engine/Plugins/Developer/Sandbox/SandboxedEditing/Source/SandboxedEditing/Private/Features/Browser/ViewModels/Transfer/ExportWorkflow.cpp b/Engine/Plugins/Developer/Sandbox/SandboxedEditing/Source/SandboxedEditing/Private/Features/Browser/ViewModels/Transfer/ExportWorkflow.cpp
new file mode 100644
--- /dev/null
+++ b/Engine/Plugins/Developer/Sandbox/SandboxedEditing/Source/SandboxedEditing/Private/Features/Browser/ViewModels/Transfer/ExportWorkflow.cpp
@@ -0,0 +1,193 @@
+﻿// Copyright Epic Games, Inc. All Rights Reserved.
+
+#include "ExportWorkflow.h"
+
+#include "ExportWorkflowResult.h"
+#include "GlobalSandboxedEditingEditorSettings.h"
+#include "Framework/Models/SandboxSystemModel.h"
+#include "Internationalization/Internationalization.h"
+#include "Misc/Paths.h"
+#include "Misc/PathViews.h"
+#include "Misc/ScopedSlowTask.h"
+#include "Utils/ZipTransferUtils.h"
+
+#define LOCTEXT_NAMESPACE "FExportWorkflow"
+
+namespace UE::SandboxedEditing
+{
+namespace ExportDetail
+{
+static bool IsDirectoryPath(const FString& InPath)
+{
+	return FPathViews::GetExtension(InPath).Len() == 0;
+}
+
+static void ExportSingle(const FString& InExportedSandboxPath, const FString& InExportPath, FExportWorkflowResult& OutResult)
+{
+	const TOptional<FString> Name = SandboxModel::GetSandboxName(InExportedSandboxPath);
+	if (!ensure(Name))
+	{
+		return;
+	}
+	
+	const bool bIsDirectory = IsDirectoryPath(InExportPath);
+	const FString FinalExportPath = bIsDirectory ? FPaths::Combine(InExportPath, *Name) : InExportPath;
+	
+	// Zip compression can take some time, so start a slow task.
+	FScopedSlowTask SlowTask(
+		1.f, 
+		Name ? FText::Format(LOCTEXT("ZipSingle.ItemFmt", "Exporting {0}"), FText::AsCultureInvariant(Name.Get(TEXT("")))) : FText::GetEmpty()
+		);
+	SlowTask.MakeDialogDelayed(0.5f);
+	
+	const bool bHasError = !FileSandboxCore::ZipUtils::ExportSandboxToZip(InExportedSandboxPath, FinalExportPath);
+	OutResult.ExportResult = bHasError ? EExportWorkflowResult::SomeErrors : EExportWorkflowResult::Success;
+	if (bHasError)
+	{
+		OutResult.SandboxesWithErrors.Add(InExportedSandboxPath);
+	}
+}
+
+static void ExportMulti(const TArray<FString>& InExportedSandboxPaths, const FString& InExportBasePath, FExportWorkflowResult& OutResult)
+{
+	bool bAnyErrors = false;
+	const auto AddError = [&OutResult, &bAnyErrors](FString InSandbox)
+	{
+		bAnyErrors = true;
+		OutResult.SandboxesWithErrors.Emplace(MoveTemp(InSandbox));
+	};
+	
+	// Zip compression can take some time, so start a slow task.
+	FScopedSlowTask SlowTask(InExportedSandboxPaths.Num());
+	SlowTask.MakeDialogDelayed(0.5f);
+	
+	for (const FString& SandboxRootPath : InExportedSandboxPaths)
+	{
+		const TOptional<FString> Name = SandboxModel::GetSandboxName(SandboxRootPath);
+		SlowTask.EnterProgressFrame(
+			1.f, 
+			Name ? FText::Format(LOCTEXT("ZipMulti.ItemFmt", "Exporting {0}"), FText::AsCultureInvariant(*Name)) : FText::GetEmpty()
+			);
+		
+		if (!ensure(Name))
+		{
+			AddError(SandboxRootPath);
+			continue;
+		}
+		
+		const FString Output = FPaths::Combine(InExportBasePath, *Name);
+		if (!FileSandboxCore::ZipUtils::ExportSandboxToZip(SandboxRootPath, Output))
+		{
+			AddError(SandboxRootPath);
+		}
+	}
+		
+	OutResult.ExportResult = bAnyErrors ? EExportWorkflowResult::SomeErrors : EExportWorkflowResult::Success;
+}
+}
+
+FExportWorkflow::FExportWorkflow(TArray<FString> InSandboxPaths, FSimpleDelegate InOnWorkflowEndedDelegate)
+	: ExportedSandboxPaths(MoveTemp(InSandboxPaths))
+	, OnWorkflowEndedDelegate(MoveTemp(InOnWorkflowEndedDelegate))
+{
+	check(!ExportedSandboxPaths.IsEmpty());
+}
+
+bool FExportWorkflow::CanExport()
+{
+	return CanSingleExport() || CanMultiExport();
+}
+
+void FExportWorkflow::TryExportOrCancel()
+{
+	if (!CanExport())
+	{
+		Cancel();
+		return;
+	}
+	
+	FExportWorkflowResult Result { .ExportResult = EExportWorkflowResult::SomeErrors };
+	
+	if (CanSingleExport())
+	{
+		ExportDetail::ExportSingle(ExportedSandboxPaths[0], ExportPath, Result);
+	}
+	else if (CanMultiExport())
+	{
+		ExportDetail::ExportMulti(ExportedSandboxPaths, ExportPath, Result);
+	}
+	
+	ExportEndedDelegate.Broadcast(Result);
+	OnWorkflowEndedDelegate.Execute();
+}
+
+void FExportWorkflow::Cancel()
+{
+	ExportEndedDelegate.Broadcast(FExportWorkflowResult{ .ExportResult = EExportWorkflowResult::Cancelled });
+	OnWorkflowEndedDelegate.Execute();
+}
+
+void FExportWorkflow::SetExportPath(FString InNewExportRoot)
+{
+	// For zip export, the system dialogue can specify a path or a full filename. But when exporting multiple items, only a directory is valid.
+	if (GetExportType() == EExportPathType::Directory 
+		&& !ensure(ExportDetail::IsDirectoryPath(InNewExportRoot)))
+	{
+		return;
+	}
+	
+	ExportPath = MoveTemp(InNewExportRoot);
+	UGlobalSandboxedEditingEditorSettings::Get()->SetDefaultExportDirectory(ExportPath);
+}
+
+const FString& FExportWorkflow::GetExportPath() const
+{
+	return ExportPath;
+}
+
+EExportPathType FExportWorkflow::GetExportType() const
+{
+	return ExportedSandboxPaths.Num() == 1 
+		? EExportPathType::Zip 
+		: EExportPathType::Directory;
+}
+
+FString FExportWorkflow::GetDefaultExportDirectory() const
+{
+	UGlobalSandboxedEditingEditorSettings* Settings = UGlobalSandboxedEditingEditorSettings::Get();
+	const FString& DefaultExport = Settings->GetDefaultExportDirectory();
+	// If the user has never exported, they may want to export to where they last imported from.
+	return DefaultExport.IsEmpty() ? Settings->GetDefaultImportDirectory() : DefaultExport;
+}
+
+FString FExportWorkflow::GetDefaultSandboxFileName() const
+{
+	if (GetExportType() == EExportPathType::Directory)
+	{
+		return TEXT("");
+	}
+	
+	if (const TOptional<FString> SandboxName = SandboxModel::GetSandboxName(ExportedSandboxPaths[0]))
+	{
+		return FPaths::SetExtension(
+			FPaths::Combine(GetDefaultExportDirectory(), *SandboxName),
+			TEXT("zip")
+			);
+	}
+	
+	return TEXT("");
+}
+
+bool FExportWorkflow::CanSingleExport() const
+{
+	return ExportedSandboxPaths.Num() == 1;
+}
+
+bool FExportWorkflow::CanMultiExport() const
+{
+	const bool bIsDirectory = ExportDetail::IsDirectoryPath(ExportPath);
+	return ExportedSandboxPaths.Num() > 1 && ensure(bIsDirectory);
+}
+}
+
+#undef LOCTEXT_NAMESPACE
diff --git a/Engine/Plugins/Developer/Sandbox/SandboxedEditing/Source/SandboxedEditing/Private/Features/Browser/ViewModels/Transfer/ExportWorkflow.h b/Engine/Plugins/Developer/Sandbox/SandboxedEditing/Source/SandboxedEditing/Private/Features/Browser/ViewModels/Transfer/ExportWorkflow.h
new file mode 100644
--- /dev/null
+++ b/Engine/Plugins/Developer/Sandbox/SandboxedEditing/Source/SandboxedEditing/Private/Features/Browser/ViewModels/Transfer/ExportWorkflow.h
@@ -0,0 +1,81 @@
+﻿// Copyright Epic Games, Inc. All Rights Reserved.
+
+#pragma once
+
+#include "Delegates/Delegate.h"
+#include "Templates/UnrealTemplate.h"
+
+namespace UE::SandboxedEditing
+{
+struct FExportWorkflowResult;
+
+DECLARE_MULTICAST_DELEGATE_OneParam(FSandboxExportPathDelegate, const FString& InDefaultName);
+DECLARE_MULTICAST_DELEGATE_OneParam(FExportEndedDelegate, const FExportWorkflowResult& InExportResult);
+
+/** What type of path the user should select. */
+enum class EExportPathType : uint8
+{
+	/** Select a directory to export multiple sandboxes to. */
+	Directory,
+	/** Select a file name for a single sandbox to be exported. */
+	Zip,
+	
+	Count
+};
+
+/** Workflow for exporting sandboxes from Unreal. */
+class FExportWorkflow : public FNoncopyable
+{
+public:
+	
+	explicit FExportWorkflow(TArray<FString> InSandboxPaths, FSimpleDelegate InOnWorkflowEndedDelegate);
+
+	/** @return Whether the sandbox selection can be exported. */
+	bool CanExport();
+	/** Exports the sandbox selection. */
+	void TryExportOrCancel();
+	
+	/** Cancels the operation. */
+	void Cancel();
+	
+	/** Sets the export path. */
+	void SetExportPath(FString InNewExportRoot);
+	/** @return The export path */
+	const FString& GetExportPath() const;
+	
+	/** @return What type of path the user should select. */
+	EExportPathType GetExportType() const;
+	
+	/** @return Default path to suggest to the user. */
+	FString GetDefaultExportDirectory() const;
+	/** @return Default filename to suggest to the user. */
+	FString GetDefaultSandboxFileName() const;
+	
+	/** @return The sandboxes that are being exported. */
+	TConstArrayView<FString> GetExportedSandboxRootPaths() const { return ExportedSandboxPaths; }
+	
+	/** Invoked when the operation ends. */
+	FExportEndedDelegate& OnExportEnded() { return ExportEndedDelegate; }
+	
+private:
+	
+	/** The sandboxes that are being exported. */
+	const TArray<FString> ExportedSandboxPaths;
+	
+	/** Invoked when the workflow ends. */
+	const FSimpleDelegate OnWorkflowEndedDelegate;
+	
+	/** Invoked when the operation ends. */
+	FExportEndedDelegate ExportEndedDelegate;
+	
+	/**
+	 * This has a different meaning depending on ExportedSandboxPaths.Num:
+	 * - Num() == 1: This is a file path to a .zip to export to, or a directory path (in which case the repository name is used as suggested name).
+	 * - Num() > 1: This is a directory path to place the .zips in.
+	 */
+	FString ExportPath;
+	
+	bool CanSingleExport() const;
+	bool CanMultiExport() const;
+};
+}
diff --git a/Engine/Plugins/Developer/Sandbox/SandboxedEditing/Source/SandboxedEditing/Private/Features/Browser/ViewModels/Transfer/ExportWorkflowResult.h b/Engine/Plugins/Developer/Sandbox/SandboxedEditing/Source/SandboxedEditing/Private/Features/Browser/ViewModels/Transfer/ExportWorkflowResult.h
new file mode 100644
--- /dev/null
+++ b/Engine/Plugins/Developer/Sandbox/SandboxedEditing/Source/SandboxedEditing/Private/Features/Browser/ViewModels/Transfer/ExportWorkflowResult.h
@@ -0,0 +1,29 @@
+﻿// Copyright Epic Games, Inc. All Rights Reserved.
+
+#pragma once
+
+#include "Delegates/Delegate.h"
+#include "Templates/UnrealTemplate.h"
+
+namespace UE::SandboxedEditing
+{
+/** Gives an indication of the workflow's result. */
+enum class EExportWorkflowResult : uint8
+{
+	Success,
+	SomeErrors,
+	Cancelled,
+	
+	Count
+};
+
+/** Describes the result of a finished export workflow. */
+struct FExportWorkflowResult
+{
+	/** Gives an indication of the workflow's result. */
+	EExportWorkflowResult ExportResult;
+	
+	/** Sandboxes that failed to be exported. */
+	TArray<FString> SandboxesWithErrors;
+};
+}
diff --git a/Engine/Plugins/Developer/Sandbox/SandboxedEditing/Source/SandboxedEditing/Private/Features/Browser/ViewModels/Transfer/ImportValidationResult.h b/Engine/Plugins/Developer/Sandbox/SandboxedEditing/Source/SandboxedEditing/Private/Features/Browser/ViewModels/Transfer/ImportValidationResult.h
new file mode 100644
--- /dev/null
+++ b/Engine/Plugins/Developer/Sandbox/SandboxedEditing/Source/SandboxedEditing/Private/Features/Browser/ViewModels/Transfer/ImportValidationResult.h
@@ -0,0 +1,22 @@
+﻿// Copyright Epic Games, Inc. All Rights Reserved.
+
+#pragma once
+
+#include "Containers/Array.h"
+#include "Containers/UnrealString.h"
+
+namespace UE::SandboxedEditing
+{
+/** Result of validating user selected files before import. */
+struct FImportValidationResult
+{
+	/**
+	 * The root paths of sandboxes that already exist and would be overwritten by the import operation. 
+	 * Ask the user to confirm overwrite. 
+	 */
+	TArray<FString> AskForOverwrite;
+	
+	/** Paths to zips that the user selected but that do not contain any valid sandboxes. */
+	TArray<FString> InvalidZips;
+};
+}
\ No newline at end of file
diff --git a/Engine/Plugins/Developer/Sandbox/SandboxedEditing/Source/SandboxedEditing/Private/Features/Browser/ViewModels/Transfer/ImportWorkflow.cpp b/Engine/Plugins/Developer/Sandbox/SandboxedEditing/Source/SandboxedEditing/Private/Features/Browser/ViewModels/Transfer/ImportWorkflow.cpp
new file mode 100644
--- /dev/null
+++ b/Engine/Plugins/Developer/Sandbox/SandboxedEditing/Source/SandboxedEditing/Private/Features/Browser/ViewModels/Transfer/ImportWorkflow.cpp
@@ -0,0 +1,163 @@
+﻿// Copyright Epic Games, Inc. All Rights Reserved.
+
+#include "ImportWorkflow.h"
+
+#include "Algo/Transform.h"
+#include "Framework/Models/SandboxSystemModel.h"
+#include "GlobalSandboxedEditingEditorSettings.h"
+#include "HAL/PlatformFileManager.h"
+#include "ImportWorkflowResult.h"
+#include "Misc/FileHelper.h"
+#include "Misc/Paths.h"
+#include "Misc/ScopedSlowTask.h"
+#include "Utils/ZipTransferUtils.h"
+
+#define LOCTEXT_NAMESPACE "FImportWorkflow"
+
+namespace UE::SandboxedEditing
+{
+namespace ImportWorkflowDetail
+{
+static bool IsDirectoryEmpty(const FString& DirectoryPath)
+{
+	IPlatformFile& PlatformFile = FPlatformFileManager::Get().GetPlatformFile();
+
+	bool bHasFiles = false;
+	PlatformFile.IterateDirectory(
+		*DirectoryPath,
+		[&bHasFiles](const TCHAR* FilenameOrDirectory, bool bIsDirectory)
+		{
+			bHasFiles = true;
+			return false; 
+		});
+
+	return !bHasFiles;
+}
+
+static FString GetImportToPath(const FString& InSandboxName)
+{
+	return FPaths::Combine(SandboxModel::GetBaseSandboxDirectory(), InSandboxName);
+}
+
+static void ValidateFilesToImport(
+	const TArray<FString>& InFilesToImport,
+	const TArray<FileSandboxCore::ZipUtils::FImportInspectionResult>& InFileInspections,
+	FImportValidationResult& OutResult
+	)
+{
+	for (int32 Index = 0; Index < InFilesToImport.Num(); ++Index)
+	{
+		const FString& FileToImport = InFilesToImport[Index];
+		const FileSandboxCore::ZipUtils::FImportInspectionResult& Inspection = InFileInspections[Index];
+		
+		if (!Inspection.IsValid())
+		{
+			OutResult.InvalidZips.Add(FileToImport);
+			continue;
+		}
+		
+		const FString OutputPath = GetImportToPath(Inspection.SandboxName);
+		if (!IsDirectoryEmpty(OutputPath))
+		{
+			OutResult.AskForOverwrite.Add(OutputPath);
+		}
+	}
+}
+} 
+
+FImportWorkflow::FImportWorkflow(FSimpleDelegate InOnWorkflowEndedDelegate)
+	: OnWorkflowEndedDelegate(MoveTemp(InOnWorkflowEndedDelegate))
+{}
+
+void FImportWorkflow::TryImportOrCancel()
+{
+	if (FilesToImport.IsEmpty())
+	{
+		Cancel();
+		return;
+	}
+	
+	// Zip compression can take some time, so start a slow task.
+	FScopedSlowTask SlowTask(FilesToImport.Num(), LOCTEXT("Import.Title", "Importing sandbox"));
+	SlowTask.MakeDialogDelayed(0.5f);
+	
+	FImportWorkflowResult Result { .ResultCode = EImportWorkflowResult::Success };
+	for (int32 Index = 0; Index < FilesToImport.Num(); ++Index)
+	{
+		const FString& FileToImport = FilesToImport[Index];
+		const FileSandboxCore::ZipUtils::FImportInspectionResult& Inspection = FileInspections[Index];
+		if (!Inspection.IsValid())
+		{
+			SlowTask.EnterProgressFrame(1.f);
+			continue;
+		}
+		
+		const FString& SandboxName = Inspection.SandboxName;
+		SlowTask.EnterProgressFrame(1.f, FText::Format(LOCTEXT("Import.ItemFmt", "Importing {0}"), FText::AsCultureInvariant(SandboxName)));
+		
+		const FString ImportToPath = ImportWorkflowDetail::GetImportToPath(SandboxName);
+		const bool bSuccess = FileSandboxCore::ZipUtils::ImportSandboxFromZip(Inspection, ImportToPath);
+		if (!bSuccess)
+		{
+			Result.ResultCode = EImportWorkflowResult::SomeErrors;
+			Result.Errors.Add(FileToImport);
+		}
+	}
+	
+	ImportEndedDelegate.Broadcast(Result);
+	OnWorkflowEndedDelegate.Execute();
+}
+
+void FImportWorkflow::Cancel()
+{
+	ImportEndedDelegate.Broadcast(FImportWorkflowResult{ .ResultCode = EImportWorkflowResult::Cancelled });
+	OnWorkflowEndedDelegate.Execute();
+}
+
+FImportValidationResult FImportWorkflow::SetFilesToImport(TArray<FString> InFiles)
+{
+	using namespace FileSandboxCore::ZipUtils;
+	if (InFiles.IsEmpty())
+	{
+		return FImportValidationResult();
+	}
+	
+	// Most likely all files are in the same directory. If not, too bad, can only have one.
+	UGlobalSandboxedEditingEditorSettings::Get()->SetDefaultImportDirectory(FilesToImport[0]);
+	
+	FImportValidationResult Result;
+	
+	FilesToImport = MoveTemp(InFiles);
+	FileInspections.Empty(FilesToImport.Num());
+	for (int32 Index = 0; Index < InFiles.Num(); ++Index)
+	{
+		const FString& File = FilesToImport[Index];
+		
+		if (TOptional<FImportInspectionResult> InspectionResult = InspectFileForImport(File))
+		{
+			FileInspections.Emplace(MoveTemp(*InspectionResult));
+		}
+		else
+		{
+			// If we're unable to read the zip, then just mark it as invalid and don't add it to the list.
+			Result.InvalidZips.Add(File);
+			
+			FilesToImport.RemoveAt(Index);
+			--Index;
+		}
+	}
+	
+	ImportWorkflowDetail::ValidateFilesToImport(FilesToImport, FileInspections, Result);
+	return Result;
+}
+
+FString FImportWorkflow::GetDefaultImportDirectory() const
+{
+	UGlobalSandboxedEditingEditorSettings* Settings = UGlobalSandboxedEditingEditorSettings::Get();
+	const FString& DefaultImport = Settings->GetDefaultImportDirectory();
+	// If the user has never imported, they may want to import from where they last exported to.
+	return DefaultImport.IsEmpty() ? Settings->GetDefaultExportDirectory() : DefaultImport;
+}
+}
+
+#undef LOCTEXT_NAMESPACE
\ No newline at end of file
diff --git a/Engine/Plugins/Developer/Sandbox/SandboxedEditing/Source/SandboxedEditing/Private/Features/Browser/ViewModels/Transfer/ImportWorkflow.h b/Engine/Plugins/Developer/Sandbox/SandboxedEditing/Source/SandboxedEditing/Private/Features/Browser/ViewModels/Transfer/ImportWorkflow.h
new file mode 100644
--- /dev/null
+++ b/Engine/Plugins/Developer/Sandbox/SandboxedEditing/Source/SandboxedEditing/Private/Features/Browser/ViewModels/Transfer/ImportWorkflow.h
@@ -0,0 +1,56 @@
+﻿// Copyright Epic Games, Inc. All Rights Reserved.
+
+#pragma once
+
+#include "Delegates/Delegate.h"
+#include "FileUtilities/ZipArchiveReader.h"
+#include "ImportValidationResult.h"
+#include "Templates/UnrealTemplate.h"
+#include "Utils/ZipTransferUtils.h"
+
+namespace UE::FileSandboxCore::ZipUtils
+{
+struct FImportInspectionResult;
+}
+
+namespace UE::SandboxedEditing
+{
+struct FImportWorkflowResult;
+DECLARE_MULTICAST_DELEGATE_OneParam(FImportEndedDelegate, const FImportWorkflowResult& InImportResult);
+
+/** Workflow for importing sandboxes into Unreal. */
+class FImportWorkflow : public FNoncopyable
+{
+public:
+	
+	explicit FImportWorkflow(FSimpleDelegate InOnWorkflowEndedDelegate);
+	
+	void TryImportOrCancel();
+	
+	void Cancel();
+	
+	/** Sets the files that are supposed */
+	FImportValidationResult SetFilesToImport(TArray<FString> InFiles);
+	/** @return The files that will be imported. */
+	const TArray<FString>& GetFilesToImport() const { return FilesToImport; }
+	
+	/** @return Default path to suggest to the user. */
+	FString GetDefaultImportDirectory() const;
+	
+	/** Invoked when the operation ends. */
+	FImportEndedDelegate& OnImportEnded() { return ImportEndedDelegate; }
+	
+private:
+	
+	/** Invoked when the workflow ends. */
+	const FSimpleDelegate OnWorkflowEndedDelegate;
+	
+	/** Paths to the files the user has selected for import. */
+	TArray<FString> FilesToImport;
+	/** Validation result for each selected file. */
+	TArray<FileSandboxCore::ZipUtils::FImportInspectionResult> FileInspections;
+	
+	/** Invoked when the operation ends. */
+	FImportEndedDelegate ImportEndedDelegate;
+};
+}
diff --git a/Engine/Plugins/Developer/Sandbox/SandboxedEditing/Source/SandboxedEditing/Private/Features/Browser/ViewModels/Transfer/ImportWorkflowResult.h b/Engine/Plugins/Developer/Sandbox/SandboxedEditing/Source/SandboxedEditing/Private/Features/Browser/ViewModels/Transfer/ImportWorkflowResult.h
new file mode 100644
--- /dev/null
+++ b/Engine/Plugins/Developer/Sandbox/SandboxedEditing/Source/SandboxedEditing/Private/Features/Browser/ViewModels/Transfer/ImportWorkflowResult.h
@@ -0,0 +1,29 @@
+// Copyright Epic Games, Inc. All Rights Reserved.
+
+#pragma once
+
+#include "Containers/UnrealString.h"
+#include "HAL/Platform.h"
+
+namespace UE::SandboxedEditing
+{
+/** Gives an indication of the workflow's result. */
+enum class EImportWorkflowResult : uint8
+{
+	Success,
+	SomeErrors,
+	Cancelled,
+	
+	Count
+};
+
+/** Describes the result of the entire import operation (multiple sandboxes can be imported at the same time). */
+struct FImportWorkflowResult
+{
+	/** Gives an indication of the workflow's result. */
+	EImportWorkflowResult ResultCode;
+	
+	/** The errors that occured. If empty, the operation was successful. */
+	TArray<FString> Errors;
+};
+}
\ No newline at end of file
diff --git a/Engine/Plugins/Developer/Sandbox/SandboxedEditing/Source/SandboxedEditing/Private/Features/Browser/Views/BrowserViews.h b/Engine/Plugins/Developer/Sandbox/SandboxedEditing/Source/SandboxedEditing/Private/Features/Browser/Views/BrowserViews.h
--- a/Engine/Plugins/Developer/Sandbox/SandboxedEditing/Source/SandboxedEditing/Private/Features/Browser/Views/BrowserViews.h
+++ b/Engine/Plugins/Developer/Sandbox/SandboxedEditing/Source/SandboxedEditing/Private/Features/Browser/Views/BrowserViews.h
@@ -7,13 +7,15 @@
 #include "PersistSandboxView.h"
 #include "RenameWorkflowView.h"
 #include "Templates/SharedPointer.h"
+#include "Transfer/ExportView.h"
+#include "Transfer/ImportView.h"
 
 namespace UE::SandboxedEditing
 {
 class FLeaveSandboxViewModel;
 
 /** 
- * Holds the views for notifications that happen in Sandboxed Editing.
+ * Holds the views for Sandboxed Editing. This is mostly for handling displaying notifications, spawning dialogues, etc. requested by view models.
  * @note Only contains "explicit" views. Of course, widgets, e.g. SSandboxedEditingRoot, implicitly classify as views. Those are not held here though.
  */
 struct FBrowserViews
@@ -30,6 +32,11 @@ struct FBrowserViews
 	/** Displays a notification when a sandbox is renamed. */
 	const TSharedRef<FRenameWorkflowView> RenameWorkflowView;
 	
+	/** Presents the UI for exporting sandboxes. */
+	const TSharedRef<FExportView> ExportView;
+	/** Presents the UI for importing sandboxes. */
+	const TSharedRef<FImportView> ImportView;
+	
 	explicit FBrowserViews(
 		const TSharedRef<FSandboxControlsViewModel>& InControlsViewModel, 
 		const TSharedRef<FPersistSandboxViewModel>& InPersistViewModel, 
@@ -39,6 +46,8 @@ struct FBrowserViews
 		, AskUserToPersistView(MakeShared<FAskUserToPersistView>(InLeaveViewModel))
 		, PersistSandboxView(MakeShared<FPersistSandboxView>(InPersistViewModel))
 		, RenameWorkflowView(MakeShared<FRenameWorkflowView>(InControlsViewModel))
+		, ExportView(MakeShared<FExportView>(InControlsViewModel))
+		, ImportView(MakeShared<FImportView>(InControlsViewModel))
 	{}
 };
 }
diff --git a/Engine/Plugins/Developer/Sandbox/SandboxedEditing/Source/SandboxedEditing/Private/Features/Browser/Views/Transfer/ExportView.cpp b/Engine/Plugins/Developer/Sandbox/SandboxedEditing/Source/SandboxedEditing/Private/Features/Browser/Views/Transfer/ExportView.cpp
new file mode 100644
--- /dev/null
+++ b/Engine/Plugins/Developer/Sandbox/SandboxedEditing/Source/SandboxedEditing/Private/Features/Browser/Views/Transfer/ExportView.cpp
@@ -0,0 +1,151 @@
+﻿// Copyright Epic Games, Inc. All Rights Reserved.
+
+#include "ExportView.h"
+
+#include "DesktopPlatformModule.h"
+#include "Features/Browser/ViewModels/SandboxControlsViewModel.h"
+#include "Features/Browser/ViewModels/Transfer/ExportWorkflowResult.h"
+#include "Framework/Application/SlateApplication.h"
+#include "Framework/Notifications/NotificationManager.h"
+#include "Widgets/Notifications/SNotificationList.h"
+
+#define LOCTEXT_NAMESPACE "FExportView"
+
+namespace UE::SandboxedEditing
+{
+FExportView::FExportView(const TSharedRef<FSandboxControlsViewModel>& InViewModel)
+	: ViewModel(InViewModel)
+{
+	ViewModel->OnExportWorkflowStarted().AddRaw(this, &FExportView::OnExportWorkflowStarted);
+}
+
+FExportView::~FExportView()
+{
+	// This cleanup should not matter effectively, the point is to follow proper RAII out of principle.
+	ViewModel->OnExportWorkflowStarted().RemoveAll(this);
+	
+	if (FExportWorkflow* Workflow = ViewModel->GetCurrentExportWorkflow())
+	{
+		Workflow->OnExportEnded().RemoveAll(this);
+	}
+}
+
+namespace ExportViewDetail
+{
+static TOptional<FString> AskUserForExportPath(FExportWorkflow& InWorkflow)
+{
+	IDesktopPlatform* DesktopPlatform = FDesktopPlatformModule::Get();
+	if (!ensure(DesktopPlatform))
+	{
+		return {};
+	}
+	
+	const EExportPathType ExportType = InWorkflow.GetExportType();
+	static_assert(static_cast<int32>(EExportPathType::Count) == 2, "Update this logic.");
+	
+	if (ExportType == EExportPathType::Zip)
+	{
+		TArray<FString> SelectedFileNames;
+		const bool bSuccess = DesktopPlatform->SaveFileDialog(
+			FSlateApplication::Get().FindBestParentWindowHandleForDialogs(nullptr),
+			LOCTEXT("ExportSandbox", "Export Sandbox").ToString(),
+			InWorkflow.GetDefaultExportDirectory(),
+			InWorkflow.GetDefaultSandboxFileName(),
+			TEXT("*.zip"),
+			EFileDialogFlags::None,
+			SelectedFileNames
+			);
+		return bSuccess && SelectedFileNames.Num() == 1 ? SelectedFileNames[0] : TOptional<FString>();
+	}
+	
+	FString SelectedDirectory;
+	const bool bSuccess = DesktopPlatform->OpenDirectoryDialog(
+		FSlateApplication::Get().FindBestParentWindowHandleForDialogs(nullptr),
+		LOCTEXT("ExportSandboxes", "Export Sandboxes").ToString(),
+		InWorkflow.GetDefaultExportDirectory(),
+		SelectedDirectory
+		); 
+	return bSuccess? SelectedDirectory : TOptional<FString>();
+}
+}
+
+void FExportView::OnExportWorkflowStarted(FExportWorkflow& InWorkflow)
+{
+	InWorkflow.OnExportEnded().AddRaw(this, &FExportView::OnExportEnded);
+	
+	if (const TOptional<FString> ExportPath = ExportViewDetail::AskUserForExportPath(InWorkflow))
+	{
+		InWorkflow.SetExportPath(*ExportPath);
+		InWorkflow.TryExportOrCancel();
+	}
+	else
+	{
+		InWorkflow.Cancel();
+	}
+}
+
+namespace ExportViewDetail
+{
+constexpr float FadeOutTime = 4.f;
+
+static void ShowSuccess(TConstArrayView<FString> InExportedSandboxes, const FString& InExportPath)
+{
+	FSlateNotificationManager& NotificationManager = FSlateNotificationManager::Get();
+	const bool bOnlyOne = InExportedSandboxes.Num() == 1;
+	const FText Title = bOnlyOne 
+		? FText::Format(
+			LOCTEXT("Success.SingleFmt", "Exported {0}"),
+			FText::AsCultureInvariant(SandboxModel::GetSandboxName(InExportedSandboxes[0]).Get(TEXT("sandbox")))
+			)
+		: FText::Format(LOCTEXT("Success.MultipleFmt", "Exported {0} {0}|plural(one=sandbox,other=sandboxes)"), InExportedSandboxes.Num());
+	
+	FNotificationInfo NotificationInfo(Title);
+	NotificationInfo.SubText = bOnlyOne 
+		? FText::Format(LOCTEXT("Success.SingleSubtextFmt", "Exported to {0}"), FText::AsCultureInvariant(InExportPath)) 
+		: FText::GetEmpty();
+	NotificationInfo.ExpireDuration = FadeOutTime;
+	NotificationInfo.FadeOutDuration = FadeOutTime;
+	NotificationManager.AddNotification(NotificationInfo)->SetCompletionState(SNotificationItem::CS_Success);
+}
+
+static void ShowWithErrors(const FExportWorkflowResult& InResult, int32 InNumExported)
+{
+	FSlateNotificationManager& NotificationManager = FSlateNotificationManager::Get();
+	FNotificationInfo NotificationInfo(LOCTEXT("WithErrors.Title", "Exported with errors"));
+	NotificationInfo.SubText = FText::Format(
+		LOCTEXT("WithErrors.SubtextFmt", "Exported {0} {0}|plural(one=sandbox,other=sandboxes).\nFailed to export {1}{1}|plural(one=sandbox,other=sandboxes)"), 
+		InNumExported - InResult.SandboxesWithErrors.Num(),
+		InResult.SandboxesWithErrors.Num()
+		);
+	NotificationInfo.ExpireDuration = FadeOutTime;
+	NotificationInfo.FadeOutDuration = FadeOutTime;
+	NotificationManager.AddNotification(NotificationInfo)->SetCompletionState(SNotificationItem::CS_Fail);
+}
+
+static void ShowCancel()
+{
+	FSlateNotificationManager& NotificationManager = FSlateNotificationManager::Get();
+	FNotificationInfo NotificationInfo(LOCTEXT("Cancelled", "Cancelled sandbox export"));
+	NotificationInfo.ExpireDuration = FadeOutTime;
+	NotificationInfo.FadeOutDuration = FadeOutTime;
+	NotificationManager.AddNotification(NotificationInfo);
+}
+}
+
+void FExportView::OnExportEnded(const FExportWorkflowResult& InResult)
+{
+	const FExportWorkflow* Workflow = ViewModel->GetCurrentExportWorkflow();
+	check(Workflow);
+	
+	static_assert(static_cast<int32>(EExportWorkflowResult::Count) == 3, "Update this switch");
+	switch (InResult.ExportResult)
+	{
+	case EExportWorkflowResult::Success: ExportViewDetail::ShowSuccess(Workflow->GetExportedSandboxRootPaths(), Workflow->GetExportPath()); break;
+	case EExportWorkflowResult::SomeErrors: ExportViewDetail::ShowWithErrors(InResult, Workflow->GetExportedSandboxRootPaths().Num()); break;
+	case EExportWorkflowResult::Cancelled: ExportViewDetail::ShowCancel(); break;
+	default: break;
+	}
+}
+}
+
+#undef LOCTEXT_NAMESPACE
diff --git a/Engine/Plugins/Developer/Sandbox/SandboxedEditing/Source/SandboxedEditing/Private/Features/Browser/Views/Transfer/ExportView.h b/Engine/Plugins/Developer/Sandbox/SandboxedEditing/Source/SandboxedEditing/Private/Features/Browser/Views/Transfer/ExportView.h
new file mode 100644
--- /dev/null
+++ b/Engine/Plugins/Developer/Sandbox/SandboxedEditing/Source/SandboxedEditing/Private/Features/Browser/Views/Transfer/ExportView.h
@@ -0,0 +1,30 @@
+﻿// Copyright Epic Games, Inc. All Rights Reserved.
+
+#pragma once
+
+#include "Features/Browser/ViewModels/Transfer/ExportWorkflow.h"
+#include "Templates/SharedPointer.h"
+#include "Templates/UnrealTemplate.h"
+
+namespace UE::SandboxedEditing
+{
+class FSandboxControlsViewModel;
+
+/** Presents dialogues for exporting sandboxes. */
+class FExportView : public FNoncopyable
+{
+public:
+	
+	explicit FExportView(const TSharedRef<FSandboxControlsViewModel>& InViewModel);
+	~FExportView();
+	
+private:
+	
+	/** Produces the export workflow that the view binds to. */
+	const TSharedRef<FSandboxControlsViewModel> ViewModel;
+
+	void OnExportWorkflowStarted(FExportWorkflow& InWorkflow);
+	void OnExportEnded(const FExportWorkflowResult& InResult);
+};
+}
+
diff --git a/Engine/Plugins/Developer/Sandbox/SandboxedEditing/Source/SandboxedEditing/Private/Features/Browser/Views/Transfer/ImportView.cpp b/Engine/Plugins/Developer/Sandbox/SandboxedEditing/Source/SandboxedEditing/Private/Features/Browser/Views/Transfer/ImportView.cpp
new file mode 100644
--- /dev/null
+++ b/Engine/Plugins/Developer/Sandbox/SandboxedEditing/Source/SandboxedEditing/Private/Features/Browser/Views/Transfer/ImportView.cpp
@@ -0,0 +1,225 @@
+﻿// Copyright Epic Games, Inc. All Rights Reserved.
+
+#include "ImportView.h"
+
+#include "DesktopPlatformModule.h"
+#include "IDesktopPlatform.h"
+#include "Algo/RemoveIf.h"
+#include "Features/Browser/ViewModels/SandboxControlsViewModel.h"
+#include "Features/Browser/ViewModels/Transfer/ImportWorkflowResult.h"
+#include "Framework/Application/SlateApplication.h"
+#include "Framework/Notifications/NotificationManager.h"
+#include "Misc/MessageDialog.h"
+#include "Widgets/Notifications/SNotificationList.h"
+
+#define LOCTEXT_NAMESPACE "FImportView"
+
+namespace UE::SandboxedEditing
+{
+FImportView::FImportView(const TSharedRef<FSandboxControlsViewModel>& InViewModel)
+	: ViewModel(InViewModel)
+{
+	ViewModel->OnImportWorkflowStarted().AddRaw(this, &FImportView::OnImportWorkflowStarted);
+}
+
+FImportView::~FImportView()
+{
+	// This cleanup should not matter effectively, the point is to follow proper RAII out of principle.
+	ViewModel->OnImportWorkflowStarted().RemoveAll(this);
+	
+	if (FImportWorkflow* Workflow = ViewModel->GetCurrentImportWorkflow())
+	{
+		Workflow->OnImportEnded().RemoveAll(this);
+	}
+}
+
+namespace ImportViewDetail
+{
+static TArray<FString> AskUserForFilesToImport(FImportWorkflow& InWorkflow)
+{
+	TArray<FString> Result;
+	
+	IDesktopPlatform* DesktopPlatform = FDesktopPlatformModule::Get();
+	if (!ensure(DesktopPlatform))
+	{
+		return Result;
+	}
+	
+	DesktopPlatform->OpenFileDialog(
+		FSlateApplication::Get().FindBestParentWindowHandleForDialogs(nullptr),
+			LOCTEXT("ImportSandbox", "Import Sandbox").ToString(),
+			InWorkflow.GetDefaultImportDirectory(),
+			TEXT(""),
+			TEXT("*.zip"),
+			EFileDialogFlags::Multiple,
+			Result
+		);
+	return Result;
+}
+
+/** @return List of sandbox names */
+static FString MakeSandboxNameList(const TArray<FString>& InStrings)
+{
+	return FString::JoinBy(InStrings, TEXT("\n"), [](const FString& SandboxPath)
+	{
+		const FString& PathLeaf = FPaths::GetPathLeaf(SandboxPath);
+		return FString::Printf(TEXT(" - %s"), *PathLeaf);
+	});
+}
+
+static void ShowInvalidZipSelected(const TArray<FString>& InFilesToImport, const TArray<FString>& InInvalidZips)
+{
+	const FText AllInvalid = FText::Format(
+		LOCTEXT("Invalid.Message.All", "The selected {0}|plural(one=file does,other=files do) not contain any {0}|plural(one=sandbox,other=sandboxes)."),
+		InFilesToImport.Num()
+		);
+	
+	const FText SomeInvalid = FText::Format(
+		LOCTEXT("Invalid.Message.Some", "The following {0}|plural(one=file does,other=files do) not contain any sandbox."),
+		InInvalidZips.Num(),
+		FText::AsCultureInvariant(MakeSandboxNameList(InInvalidZips))
+		);
+	
+	FMessageDialog::Open(
+		EAppMsgType::Ok, 
+		InFilesToImport.Num() == 1 || InFilesToImport.Num() == InInvalidZips.Num() ? AllInvalid : SomeInvalid,
+		LOCTEXT("Invalid.Title", "Cannot import sandboxes")
+		);
+}
+
+static void AskUserToOverwriteAndProceed(
+	FImportWorkflow& InWorkflow, TArray<FString>& InFilesToImport, const FImportValidationResult InValidationResult
+	)
+{
+	const FText Sandboxes = FText::AsCultureInvariant(MakeSandboxNameList(InValidationResult.AskForOverwrite));
+	const EAppReturnType::Type OverwriteFilesResponse = FMessageDialog::Open(
+		EAppMsgType::YesNoCancel,
+		FText::Format(LOCTEXT("Overwrite.MessageFmt", "The following sandboxes already exist:\n{0}\n\nDo you want to overwrite them?"), Sandboxes),
+		LOCTEXT("Overwrite.Title", "Overwrite sandboxes")
+		);
+	
+	switch (OverwriteFilesResponse)
+	{
+	case EAppReturnType::No:
+		InFilesToImport.SetNum(Algo::RemoveIf(InFilesToImport, [&InValidationResult](const FString& InFile)
+		{
+			return InValidationResult.AskForOverwrite.Contains(InFile);
+		}));
+		InWorkflow.SetFilesToImport(InFilesToImport);
+		InWorkflow.TryImportOrCancel();
+		break;
+	case EAppReturnType::Yes: 
+		InWorkflow.TryImportOrCancel();
+		break;
+	case EAppReturnType::Cancel:
+		InWorkflow.Cancel();
+		break;
+	default: checkNoEntry();
+	}
+}
+}
+
+void FImportView::OnImportWorkflowStarted(FImportWorkflow& InWorkflow)
+{
+	InWorkflow.OnImportEnded().AddRaw(this, &FImportView::OnImportEnded);
+	
+	TArray<FString> FilesToImport = ImportViewDetail::AskUserForFilesToImport(InWorkflow);
+	if (FilesToImport.IsEmpty())
+	{
+		InWorkflow.Cancel();
+		return;
+	}
+	
+	const FImportValidationResult ValidationResult = InWorkflow.SetFilesToImport(FilesToImport);
+	const TArray<FString>& InvalidZips = ValidationResult.InvalidZips;
+	if (!InvalidZips.IsEmpty())
+	{
+		ImportViewDetail::ShowInvalidZipSelected(FilesToImport, InvalidZips);
+		InWorkflow.Cancel();
+		return;
+	}
+	
+	if (!ValidationResult.AskForOverwrite.IsEmpty())
+	{
+		ImportViewDetail::AskUserToOverwriteAndProceed(InWorkflow, FilesToImport, ValidationResult);
+		return;
+	}
+	
+	InWorkflow.TryImportOrCancel();
+}
+
+namespace ImportViewDetail
+{
+constexpr float FadeOutTime = 4.f;
+
+static void ShowSuccess(TConstArrayView<FString> InImportedSandboxed)
+{
+	FSlateNotificationManager& NotificationManager = FSlateNotificationManager::Get();
+	const bool bOnlyOne = InImportedSandboxed.Num() == 1;
+	const FText Title = bOnlyOne 
+		? FText::Format(
+			LOCTEXT("Success.SingleFmt", "Imported {0}"),
+			FText::AsCultureInvariant(SandboxModel::GetSandboxName(InImportedSandboxed[0]).Get(TEXT("sandbox")))
+			)
+		: FText::Format(LOCTEXT("Success.MultipleFmt", "Imported {0} {0}|plural(one=sandbox,other=sandboxes)"), InImportedSandboxed.Num());
+	
+	FNotificationInfo NotificationInfo(Title);
+	NotificationInfo.SubText = bOnlyOne 
+		? FText::Format(LOCTEXT("Success.SingleSubtextFmt", "Imported from {0}"), FText::AsCultureInvariant(InImportedSandboxed[0])) 
+		: FText::GetEmpty();
+	NotificationInfo.ExpireDuration = FadeOutTime;
+	NotificationInfo.FadeOutDuration = FadeOutTime;
+	NotificationManager.AddNotification(NotificationInfo)->SetCompletionState(SNotificationItem::CS_Success);
+}
+
+static void ShowWithErrors(const FImportWorkflowResult& InResult, TConstArrayView<FString> InFilesToImport)
+{
+	const bool bOnlyOne = InFilesToImport.Num() == 1;
+	const FString& FirstSandboxName = FPaths::GetPathLeaf(InFilesToImport[0]);
+	
+	const FText Title = bOnlyOne 
+		? FText::Format(LOCTEXT("WithErrors.TitleOnlyOne", "Failed to import {0}"), FText::AsCultureInvariant(FirstSandboxName)) 
+		: LOCTEXT("WithErrors.Title", "Imported with errors");
+	const FText SubText = bOnlyOne 
+		? FText::GetEmpty()
+		: FText::Format(
+			LOCTEXT("WithErrors.SubtextFmt", "Imported {0} {0}|plural(one=sandbox,other=sandboxes).\nFailed to import {1}{1}|plural(one=sandbox,other=sandboxes)"), 
+			InFilesToImport.Num() - InResult.Errors.Num(),
+			InResult.Errors.Num()
+			);
+	
+	FSlateNotificationManager& NotificationManager = FSlateNotificationManager::Get();
+	FNotificationInfo NotificationInfo(Title);
+	NotificationInfo.SubText = SubText;
+	NotificationInfo.ExpireDuration = FadeOutTime;
+	NotificationInfo.FadeOutDuration = FadeOutTime;
+	NotificationManager.AddNotification(NotificationInfo)->SetCompletionState(SNotificationItem::CS_Fail);
+}
+
+static void ShowCancel()
+{
+	FSlateNotificationManager& NotificationManager = FSlateNotificationManager::Get();
+	FNotificationInfo NotificationInfo(LOCTEXT("Cancelled", "Cancelled sandbox import"));
+	NotificationInfo.ExpireDuration = FadeOutTime;
+	NotificationInfo.FadeOutDuration = FadeOutTime;
+	NotificationManager.AddNotification(NotificationInfo);
+}
+}
+
+void FImportView::OnImportEnded(const FImportWorkflowResult& InImportResult)
+{
+	const FImportWorkflow* Workflow = ViewModel->GetCurrentImportWorkflow();
+	check(Workflow);
+	
+	static_assert(static_cast<int32>(EImportWorkflowResult::Count) == 3, "Update this switch");
+	switch (InImportResult.ResultCode)
+	{
+	case EImportWorkflowResult::Success: ImportViewDetail::ShowSuccess(Workflow->GetFilesToImport()); break;
+	case EImportWorkflowResult::SomeErrors: ImportViewDetail::ShowWithErrors(InImportResult, Workflow->GetFilesToImport()); break;
+	case EImportWorkflowResult::Cancelled: ImportViewDetail::ShowCancel(); break;
+	default: checkNoEntry();
+	}
+}
+}
+
+#undef LOCTEXT_NAMESPACE
\ No newline at end of file
diff --git a/Engine/Plugins/Developer/Sandbox/SandboxedEditing/Source/SandboxedEditing/Private/Features/Browser/Views/Transfer/ImportView.h b/Engine/Plugins/Developer/Sandbox/SandboxedEditing/Source/SandboxedEditing/Private/Features/Browser/Views/Transfer/ImportView.h
new file mode 100644
--- /dev/null
+++ b/Engine/Plugins/Developer/Sandbox/SandboxedEditing/Source/SandboxedEditing/Private/Features/Browser/Views/Transfer/ImportView.h
@@ -0,0 +1,32 @@
+﻿// Copyright Epic Games, Inc. All Rights Reserved.
+
+#pragma once
+
+#include "Features/Browser/ViewModels/Transfer/ImportWorkflow.h"
+#include "Templates/SharedPointer.h"
+#include "Templates/UnrealTemplate.h"
+
+namespace UE::SandboxedEditing
+{
+class FImportWorkflow;
+class FSandboxControlsViewModel;
+struct FImportWorkflowResult;
+
+/** Presents dialogues for importing sandboxes. */
+class FImportView : public FNoncopyable
+{
+public:
+	
+	explicit FImportView(const TSharedRef<FSandboxControlsViewModel>& InViewModel);
+	~FImportView();
+	
+private:
+	
+	/** Produces the import workflow that the view binds to. */
+	const TSharedRef<FSandboxControlsViewModel> ViewModel;
+	
+	void OnImportWorkflowStarted(FImportWorkflow& InWorkflow);
+	void OnImportEnded(const FImportWorkflowResult& InImportResult);
+};
+}
+
diff --git a/Engine/Plugins/Developer/Sandbox/SandboxedEditing/Source/SandboxedEditing/Private/Features/Browser/Widgets/Unsandboxed/List/SSandboxListView.cpp b/Engine/Plugins/Developer/Sandbox/SandboxedEditing/Source/SandboxedEditing/Private/Features/Browser/Widgets/Unsandboxed/List/SSandboxListView.cpp
--- a/Engine/Plugins/Developer/Sandbox/SandboxedEditing/Source/SandboxedEditing/Private/Features/Browser/Widgets/Unsandboxed/List/SSandboxListView.cpp
+++ b/Engine/Plugins/Developer/Sandbox/SandboxedEditing/Source/SandboxedEditing/Private/Features/Browser/Widgets/Unsandboxed/List/SSandboxListView.cpp
@@ -90,6 +90,16 @@ void SSandboxListView::ForEachSelectedItem(TFunctionRef<EBreakBehavior(const FSt
 	}
 }
 
+TArray<FString> SSandboxListView::GetSelectedSandboxRootPaths() const
+{
+	TArray<FString> Result;
+	Algo::Transform(ListView->GetSelectedItems(), Result, [](const TSharedPtr<FSandboxListItem>& Item)
+	{
+		return Item->SandboxInfo.SandboxRoot;
+	});
+	return Result;
+}
+
 TOptional<FString> SSandboxListView::GetSelectedItem() const
 {
 	return NumSelected() == 1 ? ListView->GetSelectedItems()[0]->SandboxInfo.SandboxRoot : TOptional<FString>{};
@@ -153,13 +163,28 @@ TSharedPtr<SWidget> SSandboxListView::OnContextMenuOpening()
 {
 	FMenuBuilder MenuBuilder(true, ContextMenuCommandList);
 	
-	MenuBuilder.AddMenuEntry(FGenericCommands::Get().Delete);
-	
 	TArray<TSharedPtr<FSandboxListItem>> Selection = ListView->GetSelectedItems();
-	if (Selection.Num() == 1)
+	MenuBuilder.BeginSection(TEXT("Sandbox"), LOCTEXT("Sandbox", "Sandbox"));
 	{
-		MenuBuilder.AddMenuEntry(FGenericCommands::Get().Rename);
+		MenuBuilder.AddMenuEntry(FGenericCommands::Get().Delete);
+	
+		if (Selection.Num() == 1)
+		{
+			MenuBuilder.AddMenuEntry(FGenericCommands::Get().Rename);
+		}
+	}
+	MenuBuilder.EndSection();
+	
+	MenuBuilder.BeginSection(TEXT("Transfer"), LOCTEXT("Transfer", "Transfer"));
+	{
+		FBrowserCommands& BrowserCommands = FBrowserCommands::Get();
+		if (!Selection.IsEmpty())
+		{
+			MenuBuilder.AddMenuEntry(BrowserCommands.ExportSandboxes);
+		}
+		MenuBuilder.AddMenuEntry(BrowserCommands.ImportSandboxes);
 	}
+	MenuBuilder.EndSection();
 	
 	return MenuBuilder.MakeWidget();
 }
diff --git a/Engine/Plugins/Developer/Sandbox/SandboxedEditing/Source/SandboxedEditing/Private/Features/Browser/Widgets/Unsandboxed/List/SSandboxListView.h b/Engine/Plugins/Developer/Sandbox/SandboxedEditing/Source/SandboxedEditing/Private/Features/Browser/Widgets/Unsandboxed/List/SSandboxListView.h
--- a/Engine/Plugins/Developer/Sandbox/SandboxedEditing/Source/SandboxedEditing/Private/Features/Browser/Widgets/Unsandboxed/List/SSandboxListView.h
+++ b/Engine/Plugins/Developer/Sandbox/SandboxedEditing/Source/SandboxedEditing/Private/Features/Browser/Widgets/Unsandboxed/List/SSandboxListView.h
@@ -52,6 +52,8 @@ public:
 	
 	/** Enumerates every selected sandbox. */
 	void ForEachSelectedItem(TFunctionRef<EBreakBehavior(const FString& InRootPath)> InCallback) const;
+	/** @return Array containing the root path of every sandboxed selected in the list view. */
+	TArray<FString> GetSelectedSandboxRootPaths() const;
 	
 	/** @return The sandbox that is selected if exactly one item is selected. */
 	TOptional<FString> GetSelectedItem() const;
diff --git a/Engine/Plugins/Developer/Sandbox/SandboxedEditing/Source/SandboxedEditing/Private/Features/Browser/Widgets/Unsandboxed/SBrowserControls.cpp b/Engine/Plugins/Developer/Sandbox/SandboxedEditing/Source/SandboxedEditing/Private/Features/Browser/Widgets/Unsandboxed/SBrowserControls.cpp
--- a/Engine/Plugins/Developer/Sandbox/SandboxedEditing/Source/SandboxedEditing/Private/Features/Browser/Widgets/Unsandboxed/SBrowserControls.cpp
+++ b/Engine/Plugins/Developer/Sandbox/SandboxedEditing/Source/SandboxedEditing/Private/Features/Browser/Widgets/Unsandboxed/SBrowserControls.cpp
@@ -38,8 +38,7 @@ TSharedRef<SWidget> SBrowserControls::BuildButtonBar(const FArguments& InArgs)
 			FSandboxedEditingStyle::Get().GetBrush("SandboxedEditing.Browser.NewSandbox"),
 			CommandList,
 			Commands.CreateNewSandbox, LOCTEXT("CreateNew", "Create a new sandbox")
-		)
-		);
+		));
 	
 	RowBuilder.AddSeparator();
 	
@@ -48,8 +47,21 @@ TSharedRef<SWidget> SBrowserControls::BuildButtonBar(const FArguments& InArgs)
 			FSandboxedEditingStyle::Get().GetBrush("SandboxedEditing.Browser.DeleteSandbox"),
 			CommandList,
 			GenericCommands.Delete, LOCTEXT("Delete", "Delete sandbox")
-		)
-		);
+		));
+	
+	FBrowserCommands& BrowserCommands = FBrowserCommands::Get();
+	RowBuilder.AddWidget(
+		MakeIconButtonByCommand(
+			FSandboxedEditingStyle::Get().GetBrush("SandboxedEditing.Browser.ExportSandboxes"),
+			CommandList,
+			BrowserCommands.ExportSandboxes
+		));
+	RowBuilder.AddWidget(
+		MakeIconButtonByCommand(
+			FSandboxedEditingStyle::Get().GetBrush("SandboxedEditing.Browser.ImportSandboxes"),
+			CommandList,
+			BrowserCommands.ImportSandboxes
+		));
 	
 	return RowBuilder.MakeWidget();
 }
diff --git a/Engine/Plugins/Developer/Sandbox/SandboxedEditing/Source/SandboxedEditing/Private/Features/Browser/Widgets/Unsandboxed/SSandboxedEditingBrowser.cpp b/Engine/Plugins/Developer/Sandbox/SandboxedEditing/Source/SandboxedEditing/Private/Features/Browser/Widgets/Unsandboxed/SSandboxedEditingBrowser.cpp
--- a/Engine/Plugins/Developer/Sandbox/SandboxedEditing/Source/SandboxedEditing/Private/Features/Browser/Widgets/Unsandboxed/SSandboxedEditingBrowser.cpp
+++ b/Engine/Plugins/Developer/Sandbox/SandboxedEditing/Source/SandboxedEditing/Private/Features/Browser/Widgets/Unsandboxed/SSandboxedEditingBrowser.cpp
@@ -12,6 +12,7 @@
 #include "SBrowserControls.h"
 #include "SBrowserFilters.h"
 #include "SBrowserSandboxDetails.h"
+#include "Features/Browser/Commands/BrowserCommands.h"
 #include "Utils/BreakBehavior.h"
 #include "Widgets/Text/STextBlock.h"
 
@@ -71,6 +72,13 @@ void SSandboxedEditingBrowser::BindCommands()
 		GenericCommands.Rename,
 		FExecuteAction::CreateSP(this, &SSandboxedEditingBrowser::HandleRenameSelection)
 		);
+	
+	FBrowserCommands& BrowserCommands = FBrowserCommands::Get();
+	CommandList->MapAction(BrowserCommands.ExportSandboxes,
+		FExecuteAction::CreateRaw(this, &SSandboxedEditingBrowser::HandleExportSandboxes),
+		FCanExecuteAction::CreateRaw(this, &SSandboxedEditingBrowser::CanExportSandboxes)
+		);
+	CommandList->MapAction(BrowserCommands.ImportSandboxes, FExecuteAction::CreateRaw(this, &SSandboxedEditingBrowser::HandleImportSandboxes));
 }
 
 void SSandboxedEditingBrowser::HandleDeleteSelection() const
@@ -121,6 +129,22 @@ void SSandboxedEditingBrowser::HandleRenameSelection() const
 	}
 }
 
+void SSandboxedEditingBrowser::HandleExportSandboxes()
+{
+	const TArray<FString> SandboxPaths = ListView->GetSelectedSandboxRootPaths();
+	ControlsViewModel->StartExportWorkflow(SandboxPaths);
+}
+
+bool SSandboxedEditingBrowser::CanExportSandboxes()
+{
+	return ControlsViewModel->CanStartExportWorkflow() && !ListView->GetSelectedSandboxRootPaths().IsEmpty();
+}
+
+void SSandboxedEditingBrowser::HandleImportSandboxes()
+{
+	ControlsViewModel->StartImportWorkflow();
+}
+
 void SSandboxedEditingBrowser::OnSelectedSandboxChanged(TOptional<FString> InSandbox) const
 {
 	if (InSandbox)
diff --git a/Engine/Plugins/Developer/Sandbox/SandboxedEditing/Source/SandboxedEditing/Private/Features/Browser/Widgets/Unsandboxed/SSandboxedEditingBrowser.h b/Engine/Plugins/Developer/Sandbox/SandboxedEditing/Source/SandboxedEditing/Private/Features/Browser/Widgets/Unsandboxed/SSandboxedEditingBrowser.h
--- a/Engine/Plugins/Developer/Sandbox/SandboxedEditing/Source/SandboxedEditing/Private/Features/Browser/Widgets/Unsandboxed/SSandboxedEditingBrowser.h
+++ b/Engine/Plugins/Developer/Sandbox/SandboxedEditing/Source/SandboxedEditing/Private/Features/Browser/Widgets/Unsandboxed/SSandboxedEditingBrowser.h
@@ -58,8 +58,8 @@ private:
 	
 	/** List view displaying the sandboxes. */
 	TSharedPtr<SSandboxListView> ListView;
-	
-	
+
+
 	/** Binds the commands only possible with UI, e.g. those that perform an action on the selection. */
 	void BindCommands();
 	
@@ -71,6 +71,14 @@ private:
 	/** Renames the currently selected. Only does anything if exactly one item is selected. No-op, otherwise. */
 	void HandleRenameSelection() const;
 	
+	/** Handles the command for exporting sandboxes. */
+	void HandleExportSandboxes();
+	/** @return Whether any sandboxes are selected in the browser. */
+	bool CanExportSandboxes();
+	
+	/** Handles the command for importing sandboxes. */
+	void HandleImportSandboxes();
+	
 	/** Notify the UnloadedSandboxFileStateViewModel when the selected sandbox changes. */
 	void OnSelectedSandboxChanged(TOptional<FString> InSandbox) const;
 };
diff --git a/Engine/Plugins/Developer/Sandbox/SandboxedEditing/Source/SandboxedEditing/Private/Framework/Models/SandboxSystemModel.cpp b/Engine/Plugins/Developer/Sandbox/SandboxedEditing/Source/SandboxedEditing/Private/Framework/Models/SandboxSystemModel.cpp
--- a/Engine/Plugins/Developer/Sandbox/SandboxedEditing/Source/SandboxedEditing/Private/Framework/Models/SandboxSystemModel.cpp
+++ b/Engine/Plugins/Developer/Sandbox/SandboxedEditing/Source/SandboxedEditing/Private/Framework/Models/SandboxSystemModel.cpp
@@ -481,6 +481,17 @@ TOptional<FString> SandboxModel::GetSandboxPathFor(const FString& InSandbox, con
 {
 	return FileSandboxCore::GetSandboxPathFor(InSandbox, InNonSandboxPath);
 }
+
+TOptional<FString> GetSandboxName(const FString& InSandboxDirectory)
+{
+	return FileSandboxCore::GetSandboxName(InSandboxDirectory);
+}
+
+FString GetBaseSandboxDirectory()
+{
+	// For now we use the global default directory but this could change in the future, e.g. by settings.
+	return FileSandboxCore::GetBaseSandboxDirectory();
+}
 }
 }
 
diff --git a/Engine/Plugins/Developer/Sandbox/SandboxedEditing/Source/SandboxedEditing/Private/Framework/Models/SandboxSystemModel.h b/Engine/Plugins/Developer/Sandbox/SandboxedEditing/Source/SandboxedEditing/Private/Framework/Models/SandboxSystemModel.h
--- a/Engine/Plugins/Developer/Sandbox/SandboxedEditing/Source/SandboxedEditing/Private/Framework/Models/SandboxSystemModel.h
+++ b/Engine/Plugins/Developer/Sandbox/SandboxedEditing/Source/SandboxedEditing/Private/Framework/Models/SandboxSystemModel.h
@@ -162,5 +162,11 @@ namespace SandboxModel
 {
 /** @return Gets the sandbox path for InNonSandboxPath. */
 TOptional<FString> GetSandboxPathFor(const FString& InSandbox, const FString& InNonSandboxPath);
+
+/** @return Name of the sandbox if the directory is the root of a sandbox. Unset otherwise. */
+TOptional<FString> GetSandboxName(const FString& InSandboxDirectory);
+
+/** @return The base directory in which sandboxes are placed. */
+FString GetBaseSandboxDirectory();
 }
 }
\ No newline at end of file
diff --git a/Engine/Plugins/Developer/Sandbox/SandboxedEditing/Source/SandboxedEditing/Private/GlobalSandboxedEditingEditorSettings.cpp b/Engine/Plugins/Developer/Sandbox/SandboxedEditing/Source/SandboxedEditing/Private/GlobalSandboxedEditingEditorSettings.cpp
new file mode 100644
--- /dev/null
+++ b/Engine/Plugins/Developer/Sandbox/SandboxedEditing/Source/SandboxedEditing/Private/GlobalSandboxedEditingEditorSettings.cpp
@@ -0,0 +1,39 @@
+﻿// Copyright Epic Games, Inc. All Rights Reserved.
+
+#include "GlobalSandboxedEditingEditorSettings.h"
+
+#include "Misc/Paths.h"
+#include "Misc/PathViews.h"
+
+namespace UE::SandboxedEditing::GlobalSettingsDetail
+{
+static FString GetDirectoryPath(const FString& InPath)
+{
+	FString NormalizedPath = InPath;
+	FPaths::NormalizeDirectoryName(NormalizedPath);
+	const bool bIsDirectoryPath = FPaths::GetExtension(NormalizedPath).IsEmpty();
+	return bIsDirectoryPath
+		? NormalizedPath
+		: FPaths::GetPath(NormalizedPath);
+}
+}
+
+void UGlobalSandboxedEditingEditorSettings::SetDefaultExportDirectory(const FString& InPath)
+{
+	FString DirectoryPath = UE::SandboxedEditing::GlobalSettingsDetail::GetDirectoryPath(InPath);
+	if (DirectoryPath != DefaultExportDirectory)
+	{
+		DefaultExportDirectory = MoveTemp(DirectoryPath);
+		SaveConfig();
+	}
+}
+
+void UGlobalSandboxedEditingEditorSettings::SetDefaultImportDirectory(const FString& InPath)
+{
+	FString DirectoryPath = UE::SandboxedEditing::GlobalSettingsDetail::GetDirectoryPath(InPath);
+	if (DirectoryPath != DefaultImportDirectory)
+	{
+		DefaultImportDirectory = MoveTemp(DirectoryPath);
+		SaveConfig();
+	}
+}
diff --git a/Engine/Plugins/Developer/Sandbox/SandboxedEditing/Source/SandboxedEditing/Private/GlobalSandboxedEditingEditorSettings.h b/Engine/Plugins/Developer/Sandbox/SandboxedEditing/Source/SandboxedEditing/Private/GlobalSandboxedEditingEditorSettings.h
new file mode 100644
--- /dev/null
+++ b/Engine/Plugins/Developer/Sandbox/SandboxedEditing/Source/SandboxedEditing/Private/GlobalSandboxedEditingEditorSettings.h
@@ -0,0 +1,36 @@
+﻿// Copyright Epic Games, Inc. All Rights Reserved.
+
+#pragma once
+
+#include "UObject/Object.h"
+#include "GlobalSandboxedEditingEditorSettings.generated.h"
+
+/** Settings shared with all editor instances. */
+UCLASS(Config=EditorSettings, DefaultConfig)
+class UGlobalSandboxedEditingEditorSettings : public UObject
+{
+	GENERATED_BODY()
+public:
+	
+	static UGlobalSandboxedEditingEditorSettings* Get() { return GetMutableDefault<UGlobalSandboxedEditingEditorSettings>(); }
+
+	/** @return The default path that sandboxes are exported to. */
+	const FString& GetDefaultExportDirectory() const { return DefaultExportDirectory; }
+	/** @param InPath Path to directory or file. In case of file, the directory will be extracted and set. */
+	void SetDefaultExportDirectory(const FString& InPath);
+	
+	/** @return The default path that sandboxes are imported from. */
+	const FString& GetDefaultImportDirectory() const { return DefaultImportDirectory; }
+	/** @param InPath Path to directory or file. In case of file, the directory will be extracted and set. */
+	void SetDefaultImportDirectory(const FString& InPath);
+	
+private:
+	
+	/** Default directory to export to. */
+	UPROPERTY(EditAnywhere, Config, Category = "General")
+	FString DefaultExportDirectory;
+	
+	/** Default directory to import from. */
+	UPROPERTY(EditAnywhere, Config, Category = "General")
+	FString DefaultImportDirectory;
+};
\ No newline at end of file
diff --git a/Engine/Plugins/Developer/Sandbox/SandboxedEditing/Source/SandboxedEditing/Private/LogSandboxedEditing.cpp b/Engine/Plugins/Developer/Sandbox/SandboxedEditing/Source/SandboxedEditing/Private/LogSandboxedEditing.cpp
new file mode 100644
--- /dev/null
+++ b/Engine/Plugins/Developer/Sandbox/SandboxedEditing/Source/SandboxedEditing/Private/LogSandboxedEditing.cpp
@@ -0,0 +1,5 @@
+﻿// Copyright Epic Games, Inc. All Rights Reserved.
+
+#include "LogSandboxedEditing.h"
+
+DEFINE_LOG_CATEGORY(LogSandboxedEditing);
diff --git a/Engine/Plugins/Developer/Sandbox/SandboxedEditing/Source/SandboxedEditing/Private/LogSandboxedEditing.h b/Engine/Plugins/Developer/Sandbox/SandboxedEditing/Source/SandboxedEditing/Private/LogSandboxedEditing.h
new file mode 100644
--- /dev/null
+++ b/Engine/Plugins/Developer/Sandbox/SandboxedEditing/Source/SandboxedEditing/Private/LogSandboxedEditing.h
@@ -0,0 +1,7 @@
+﻿// Copyright Epic Games, Inc. All Rights Reserved.
+
+#pragma once
+
+#include "Logging/LogCategory.h"
+
+DECLARE_LOG_CATEGORY_EXTERN(LogSandboxedEditing, Log, All);
\ No newline at end of file
diff --git a/Engine/Plugins/Developer/Sandbox/SandboxedEditing/Source/SandboxedEditing/Private/SandboxedEditingStyle.cpp b/Engine/Plugins/Developer/Sandbox/SandboxedEditing/Source/SandboxedEditing/Private/SandboxedEditingStyle.cpp
--- a/Engine/Plugins/Developer/Sandbox/SandboxedEditing/Source/SandboxedEditing/Private/SandboxedEditingStyle.cpp
+++ b/Engine/Plugins/Developer/Sandbox/SandboxedEditing/Source/SandboxedEditing/Private/SandboxedEditingStyle.cpp
@@ -2,6 +2,7 @@
 
 #include "SandboxedEditingStyle.h"
 
+#include "Styling/AppStyle.h"
 #include "Styling/SlateStyleMacros.h"
 #include "Styling/SlateStyleRegistry.h"
 
@@ -37,10 +38,17 @@ void FSandboxedEditingStyle::RegisterBrowserStyle()
 	const FLinearColor IconColorAndOpacity(FLinearColor(1.f, 1.f, 1.f, 1.f));
 	const FVector2D Icon24x24(24.0f, 24.0f);
 	
+	// Icons
 	Set("SandboxedEditing.Browser.Icon", new IMAGE_BRUSH_SVG("SandboxBrowserIcon_24", Icon24x24, IconColorAndOpacity));
 	Set("SandboxedEditing.Browser.NewSandbox", new IMAGE_BRUSH_SVG("NewSandbox_24", Icon24x24, IconColorAndOpacity));
 	Set("SandboxedEditing.Browser.DeleteSandbox", new IMAGE_BRUSH("DeleteSandbox_48", Icon24x24, IconColorAndOpacity));
 	Set("SandboxedEditing.Browser.LeaveSandbox", new IMAGE_BRUSH("LeaveSandbox_48", Icon24x24, IconColorAndOpacity));
+	Set("SandboxedEditing.Browser.ExportSandboxes", new FSlateBrush(*FAppStyle::Get().GetBrush("Icons.Export")));
+	Set("SandboxedEditing.Browser.ImportSandboxes",  new FSlateBrush(*FAppStyle::Get().GetBrush("Icons.Import")));
+	
+	// Commands
+	Set("Browser.ExportSandboxes", new FSlateBrush(*FAppStyle::Get().GetBrush("Icons.Export")));
+	Set("Browser.ImportSandboxes", new FSlateBrush(*FAppStyle::Get().GetBrush("Icons.Import")));
 	
 	// Sandbox browser
 	Set("SandboxedEditing.Browser.RowPadding", FMargin(0.f, 4.f));
diff --git a/Engine/Plugins/Developer/Sandbox/SandboxedEditing/Source/SandboxedEditing/SandboxedEditing.Build.cs b/Engine/Plugins/Developer/Sandbox/SandboxedEditing/Source/SandboxedEditing/SandboxedEditing.Build.cs
--- a/Engine/Plugins/Developer/Sandbox/SandboxedEditing/Source/SandboxedEditing/SandboxedEditing.Build.cs
+++ b/Engine/Plugins/Developer/Sandbox/SandboxedEditing/Source/SandboxedEditing/SandboxedEditing.Build.cs
@@ -18,8 +18,10 @@ namespace UnrealBuildTool.Rules
 				new string[]
 				{
 					"AssetTools",
+					"DesktopPlatform",
 					"FileSandboxCore",
 					"FileSandboxUI",
+					"FileUtilities", 
 					"InputCore",
 					"Slate", 
 					"SlateCore", 
