Sandboxed Editing: Fix importing not working due to copy-paste typo.

#Synced-CL 50904440

diff --git a/Engine/Plugins/Developer/Sandbox/SandboxedEditing/Source/SandboxedEditing/Private/Features/Browser/ViewModels/Transfer/ImportWorkflow.cpp b/Engine/Plugins/Developer/Sandbox/SandboxedEditing/Source/SandboxedEditing/Private/Features/Browser/ViewModels/Transfer/ImportWorkflow.cpp
--- a/Engine/Plugins/Developer/Sandbox/SandboxedEditing/Source/SandboxedEditing/Private/Features/Browser/ViewModels/Transfer/ImportWorkflow.cpp
+++ b/Engine/Plugins/Developer/Sandbox/SandboxedEditing/Source/SandboxedEditing/Private/Features/Browser/ViewModels/Transfer/ImportWorkflow.cpp
@@ -123,13 +123,13 @@ FImportValidationResult FImportWorkflow::SetFilesToImport(TArray<FString> InFile
 	}
 	
 	// Most likely all files are in the same directory. If not, too bad, can only have one.
-	UGlobalSandboxedEditingEditorSettings::Get()->SetDefaultImportDirectory(FilesToImport[0]);
+	UGlobalSandboxedEditingEditorSettings::Get()->SetDefaultImportDirectory(InFiles[0]);
 	
 	FImportValidationResult Result;
 	
 	FilesToImport = MoveTemp(InFiles);
 	FileInspections.Empty(FilesToImport.Num());
-	for (int32 Index = 0; Index < InFiles.Num(); ++Index)
+	for (int32 Index = 0; Index < FilesToImport.Num(); ++Index)
 	{
 		const FString& File = FilesToImport[Index];
 		
