Sandboxed Editing:
- Add icons to "Show in" actions
- Renames the label to "Show in {FileManagerName}"
- Disable the "Show in " option when not exactly one file is selected

#Synced-CL 50904440

diff --git a/Engine/Plugins/Developer/Sandbox/SandboxedEditing/Source/SandboxedEditing/Private/Features/Browser/Commands/FileStateActions/FileStateCommandBindings.cpp b/Engine/Plugins/Developer/Sandbox/SandboxedEditing/Source/SandboxedEditing/Private/Features/Browser/Commands/FileStateActions/FileStateCommandBindings.cpp
--- a/Engine/Plugins/Developer/Sandbox/SandboxedEditing/Source/SandboxedEditing/Private/Features/Browser/Commands/FileStateActions/FileStateCommandBindings.cpp
+++ b/Engine/Plugins/Developer/Sandbox/SandboxedEditing/Source/SandboxedEditing/Private/Features/Browser/Commands/FileStateActions/FileStateCommandBindings.cpp
@@ -66,7 +66,7 @@ void FFileStateCommandBindings::HandleShowFileInExplorer() const
 {
 	const TOptional<FString> SandboxRoot = SelectedSandboxRootAttr.Get();
 	const TArray<FString> SelectedPaths = SelectedFilePathsAttr.Get();
-	if (!SandboxRoot || SelectedPaths.IsEmpty())
+	if (!SandboxRoot || SelectedPaths.Num() != 1)
 	{
 		return;
 	}
@@ -81,7 +81,7 @@ bool FFileStateCommandBindings::CanShowFileInExplorer() const
 {
 	const TOptional<FString> SandboxRoot = SelectedSandboxRootAttr.Get();
 	const TArray<FString> SelectedPaths = SelectedFilePathsAttr.Get();
-	if (!SandboxRoot || SelectedPaths.IsEmpty())
+	if (!SandboxRoot || SelectedPaths.Num() != 1)
 	{
 		return false;
 	}
diff --git a/Engine/Plugins/Developer/Sandbox/SandboxedEditing/Source/SandboxedEditing/Private/Features/Browser/Commands/FileStateActions/FileStateCommands.cpp b/Engine/Plugins/Developer/Sandbox/SandboxedEditing/Source/SandboxedEditing/Private/Features/Browser/Commands/FileStateActions/FileStateCommands.cpp
--- a/Engine/Plugins/Developer/Sandbox/SandboxedEditing/Source/SandboxedEditing/Private/Features/Browser/Commands/FileStateActions/FileStateCommands.cpp
+++ b/Engine/Plugins/Developer/Sandbox/SandboxedEditing/Source/SandboxedEditing/Private/Features/Browser/Commands/FileStateActions/FileStateCommands.cpp
@@ -19,8 +19,8 @@ FFileStateCommands::FFileStateCommands()
 
 void FFileStateCommands::RegisterCommands()
 {
-	UI_COMMAND(ShowRootInExplorer, "Open sandbox folder", "Shows sandbox root directory in the system file manager.", EUserInterfaceActionType::Button, FInputChord())
-	UI_COMMAND(ShowFileInExplorer, "Open in file manager", "Shows the file's sandboxed counterpart in the system file manager.", EUserInterfaceActionType::Button, FInputChord())
+	UI_COMMAND(ShowRootInExplorer, "Show sandbox folder", "Shows sandbox root directory in the system file manager.", EUserInterfaceActionType::Button, FInputChord())
+	UI_COMMAND(ShowFileInExplorer, "Show in file manager", "Shows the file's sandboxed counterpart in the system file manager.", EUserInterfaceActionType::Button, FInputChord())
 }
 }
 
diff --git a/Engine/Plugins/Developer/Sandbox/SandboxedEditing/Source/SandboxedEditing/Private/Features/Browser/Views/Menu/FileStateMenuUtils.cpp b/Engine/Plugins/Developer/Sandbox/SandboxedEditing/Source/SandboxedEditing/Private/Features/Browser/Views/Menu/FileStateMenuUtils.cpp
--- a/Engine/Plugins/Developer/Sandbox/SandboxedEditing/Source/SandboxedEditing/Private/Features/Browser/Views/Menu/FileStateMenuUtils.cpp
+++ b/Engine/Plugins/Developer/Sandbox/SandboxedEditing/Source/SandboxedEditing/Private/Features/Browser/Views/Menu/FileStateMenuUtils.cpp
@@ -9,6 +9,8 @@
 #include "Features/Browser/Widgets/Shared/FileState/SFileStateListView.h"
 #include "Framework/MultiBox/MultiBoxBuilder.h"
 
+#define LOCTEXT_NAMESPACE "FileStateMenu"
+
 namespace UE::SandboxedEditing::FileStateMenu
 {
 TArray<FString> TransformSelectedItems(const TSharedRef<SFileStateListView>& InListView)
@@ -31,14 +33,27 @@ TAttribute<TArray<FString>> FileStateMenu::MakeSelectedFilesAttribute(const TSha
 	});
 }
 
+namespace Detail
+{
+static FText GetExploreFolderText()
+{
+	FFormatNamedArguments Args;
+	Args.Add(TEXT("FileManagerName"), FPlatformMisc::GetFileManagerName());
+	return FText::Format(NSLOCTEXT("GenericPlatform", "ShowInFileManager", "Show in {FileManagerName}"), Args);
+}
+}
+
+
 void FileStateMenu::AppendMenu(FMenuBuilder& InMenuBuilder)
 {
-	InMenuBuilder.BeginSection(TEXT("Common"));
+	InMenuBuilder.BeginSection(TEXT("Explorer"), LOCTEXT("Explorer", "Explorer"));
 	
 	FFileStateCommands& Commands = FFileStateCommands::Get();
 	InMenuBuilder.AddMenuEntry(Commands.ShowRootInExplorer);
-	InMenuBuilder.AddMenuEntry(Commands.ShowFileInExplorer);
+	InMenuBuilder.AddMenuEntry(Commands.ShowFileInExplorer, NAME_None, Detail::GetExploreFolderText());
 	
 	InMenuBuilder.EndSection();
 }
 }
+
+#undef LOCTEXT_NAMESPACE
\ No newline at end of file
diff --git a/Engine/Plugins/Developer/Sandbox/SandboxedEditing/Source/SandboxedEditing/Private/SandboxedEditingStyle.cpp b/Engine/Plugins/Developer/Sandbox/SandboxedEditing/Source/SandboxedEditing/Private/SandboxedEditingStyle.cpp
--- a/Engine/Plugins/Developer/Sandbox/SandboxedEditing/Source/SandboxedEditing/Private/SandboxedEditingStyle.cpp
+++ b/Engine/Plugins/Developer/Sandbox/SandboxedEditing/Source/SandboxedEditing/Private/SandboxedEditingStyle.cpp
@@ -50,6 +50,9 @@ void FSandboxedEditingStyle::RegisterBrowserStyle()
 	Set("Browser.ExportSandboxes", new FSlateBrush(*FAppStyle::Get().GetBrush("Icons.Export")));
 	Set("Browser.ImportSandboxes", new FSlateBrush(*FAppStyle::Get().GetBrush("Icons.Import")));
 	
+	Set("FileState.ShowRootInExplorer", new FSlateBrush(*FAppStyle::Get().GetBrush("SystemWideCommands.FindInContentBrowser")));
+	Set("FileState.ShowFileInExplorer", new FSlateBrush(*FAppStyle::Get().GetBrush("SystemWideCommands.FindInContentBrowser")));
+	
 	// Sandbox browser
 	Set("SandboxedEditing.Browser.RowPadding", FMargin(0.f, 4.f));
 	Set("SandboxedEditing.Browser.NameColumn.Padding", FMargin(16.f, 0.f, 0.f, 0.f));
