Sandboxed Editing: Add "Persist" right-click menu entry in active sandbox file change UI.

#Synced-CL 50904440

diff --git a/Engine/Plugins/Developer/Sandbox/SandboxedEditing/Source/SandboxedEditing/Private/Features/Browser/Commands/FileStateActions/FileStateCommands.cpp b/Engine/Plugins/Developer/Sandbox/SandboxedEditing/Source/SandboxedEditing/Private/Features/Browser/Commands/FileStateActions/FileStateCommands.cpp
--- a/Engine/Plugins/Developer/Sandbox/SandboxedEditing/Source/SandboxedEditing/Private/Features/Browser/Commands/FileStateActions/FileStateCommands.cpp
+++ b/Engine/Plugins/Developer/Sandbox/SandboxedEditing/Source/SandboxedEditing/Private/Features/Browser/Commands/FileStateActions/FileStateCommands.cpp
@@ -21,6 +21,8 @@ void FFileStateCommands::RegisterCommands()
 {
 	UI_COMMAND(ShowRootInExplorer, "Show sandbox folder", "Shows sandbox root directory in the system file manager.", EUserInterfaceActionType::Button, FInputChord())
 	UI_COMMAND(ShowFileInExplorer, "Show in file manager", "Shows the file's sandboxed counterpart in the system file manager.", EUserInterfaceActionType::Button, FInputChord())
+	
+	UI_COMMAND(PersistSelected, "Persist", "Open the persist dialogue with the currently selected files already pre-selected in the dialogue.", EUserInterfaceActionType::Button, FInputChord())
 }
 }
 
diff --git a/Engine/Plugins/Developer/Sandbox/SandboxedEditing/Source/SandboxedEditing/Private/Features/Browser/Commands/FileStateActions/FileStateCommands.h b/Engine/Plugins/Developer/Sandbox/SandboxedEditing/Source/SandboxedEditing/Private/Features/Browser/Commands/FileStateActions/FileStateCommands.h
--- a/Engine/Plugins/Developer/Sandbox/SandboxedEditing/Source/SandboxedEditing/Private/Features/Browser/Commands/FileStateActions/FileStateCommands.h
+++ b/Engine/Plugins/Developer/Sandbox/SandboxedEditing/Source/SandboxedEditing/Private/Features/Browser/Commands/FileStateActions/FileStateCommands.h
@@ -22,6 +22,9 @@ public:
 	
 	/** Shows the file's sandboxed counterpart in the system file explorer. */
 	TSharedPtr<FUICommandInfo> ShowFileInExplorer;
+	
+	/** Persists the selected files. */
+	TSharedPtr<FUICommandInfo> PersistSelected;
 };
 }
 
diff --git a/Engine/Plugins/Developer/Sandbox/SandboxedEditing/Source/SandboxedEditing/Private/Features/Browser/ViewModels/Persist/PersistSandboxViewModel.cpp b/Engine/Plugins/Developer/Sandbox/SandboxedEditing/Source/SandboxedEditing/Private/Features/Browser/ViewModels/Persist/PersistSandboxViewModel.cpp
--- a/Engine/Plugins/Developer/Sandbox/SandboxedEditing/Source/SandboxedEditing/Private/Features/Browser/ViewModels/Persist/PersistSandboxViewModel.cpp
+++ b/Engine/Plugins/Developer/Sandbox/SandboxedEditing/Source/SandboxedEditing/Private/Features/Browser/ViewModels/Persist/PersistSandboxViewModel.cpp
@@ -13,24 +13,24 @@ FPersistSandboxViewModel::FPersistSandboxViewModel(const TSharedRef<FSandboxSyst
 	: Model(InModel)
 {}
 
-bool FPersistSandboxViewModel::StartPersistWorkflow(FileSandboxCore::FGatheredFileChanges InFileChanges)
+bool FPersistSandboxViewModel::StartPersistWorkflow(FileSandboxCore::FGatheredFileChanges InFileChanges, TConstArrayView<FString> InPreSelected)
 {
 	if (CanStartPersistWorkflow())
 	{
 		ActivePersistWorkflow = MakeUnique<FPersistSandboxWorkflow>(Model, MoveTemp(InFileChanges),
 			FPersistSandboxWorkflow::FCleanupWorkflow::CreateRaw(this, &FPersistSandboxViewModel::OnCleanupPersistWorkflow)
 			);
-		OnStartPersistWorkflowDelegate.Broadcast(*ActivePersistWorkflow);
+		OnStartPersistWorkflowDelegate.Broadcast(*ActivePersistWorkflow, InPreSelected);
 		return true;
 	}
 	
 	return false;
 }
 
-bool FPersistSandboxViewModel::StartPersistWorkflowForActiveSandbox()
+bool FPersistSandboxViewModel::StartPersistWorkflowForActiveSandbox(TConstArrayView<FString> InPreSelected)
 {
 	FileSandboxCore::FGatheredFileChanges FileChanges = Model->GatherActiveFileChanges();
-	return FileChanges.HasChanges() && StartPersistWorkflow(MoveTemp(FileChanges));
+	return FileChanges.HasChanges() && StartPersistWorkflow(MoveTemp(FileChanges), InPreSelected);
 }
 
 bool FPersistSandboxViewModel::CanStartPersistWorkflow() const
diff --git a/Engine/Plugins/Developer/Sandbox/SandboxedEditing/Source/SandboxedEditing/Private/Features/Browser/ViewModels/Persist/PersistSandboxViewModel.h b/Engine/Plugins/Developer/Sandbox/SandboxedEditing/Source/SandboxedEditing/Private/Features/Browser/ViewModels/Persist/PersistSandboxViewModel.h
--- a/Engine/Plugins/Developer/Sandbox/SandboxedEditing/Source/SandboxedEditing/Private/Features/Browser/ViewModels/Persist/PersistSandboxViewModel.h
+++ b/Engine/Plugins/Developer/Sandbox/SandboxedEditing/Source/SandboxedEditing/Private/Features/Browser/ViewModels/Persist/PersistSandboxViewModel.h
@@ -11,7 +11,7 @@ namespace UE::FileSandboxUI { struct FPersistSummary; }
 namespace UE::SandboxedEditing
 {
 enum class EPersistWorkflowAction : uint8;
-DECLARE_MULTICAST_DELEGATE_OneParam(FPersistSandboxWorkflowDelegate, FPersistSandboxWorkflow&)
+DECLARE_MULTICAST_DELEGATE_TwoParams(FPersistSandboxWorkflowDelegate, FPersistSandboxWorkflow&, TConstArrayView<FString> /*InPreSelectedFiles*/)
 DECLARE_MULTICAST_DELEGATE_OneParam(FPersistWorkflowActionDelegate, EPersistWorkflowAction);
 DECLARE_MULTICAST_DELEGATE_OneParam(FRequestPersistSummaryNotificationDelegate, const FileSandboxUI::FPersistSummary&)
 
@@ -23,9 +23,12 @@ public:
 	explicit FPersistSandboxViewModel(const TSharedRef<FSandboxSystemModel>& InModel);
 
 	/** Starts the persist workflow. */
-	bool StartPersistWorkflow(FileSandboxCore::FGatheredFileChanges InFileChanges);
-	/** Starts a perist workflow by using the changes of the active sandbox. */
-	bool StartPersistWorkflowForActiveSandbox();
+	bool StartPersistWorkflow(FileSandboxCore::FGatheredFileChanges InFileChanges, TConstArrayView<FString> InPreSelected = {});
+	/**
+	 * Starts a persist workflow by using the changes of the active sandbox.
+	 * @param 
+	 */
+	bool StartPersistWorkflowForActiveSandbox(TConstArrayView<FString> InPreSelected = {});
 	
 	/** @return Whether a persist workflow can be started */
 	bool CanStartPersistWorkflow() const;
diff --git a/Engine/Plugins/Developer/Sandbox/SandboxedEditing/Source/SandboxedEditing/Private/Features/Browser/Views/PersistSandboxView.cpp b/Engine/Plugins/Developer/Sandbox/SandboxedEditing/Source/SandboxedEditing/Private/Features/Browser/Views/PersistSandboxView.cpp
--- a/Engine/Plugins/Developer/Sandbox/SandboxedEditing/Source/SandboxedEditing/Private/Features/Browser/Views/PersistSandboxView.cpp
+++ b/Engine/Plugins/Developer/Sandbox/SandboxedEditing/Source/SandboxedEditing/Private/Features/Browser/Views/PersistSandboxView.cpp
@@ -27,7 +27,32 @@ FPersistSandboxView::~FPersistSandboxView()
 	PersistViewModel->OnRequestPersistSummaryNotification().RemoveAll(this);
 }
 
-void FPersistSandboxView::ShowPersistWorkflow(FPersistSandboxWorkflow& InWorkflow) const
+namespace PersistViewDetail
+{
+static void PreSelectFiles(FPersistOperationViewModel& InViewModel, TConstArrayView<FString> InPreSelectedFiles)
+{
+	if (!InPreSelectedFiles.IsEmpty())
+	{
+		InViewModel.SetAllFilesPersisted(false);
+		
+		const TArray<FString>& NonSandboxPaths = InViewModel.GetPersistableFiles().NonSandboxPaths;
+		for (const FString& PreSelected : InPreSelectedFiles)
+		{
+			const int32 Index = NonSandboxPaths.IndexOfByPredicate([&PreSelected](const FString& Path)
+			{
+				return FPaths::IsSamePath(Path, PreSelected);
+			});
+			
+			if (ensure(NonSandboxPaths.IsValidIndex(Index)))
+			{
+				InViewModel.SetFilePersisted(Index, true);
+			}
+		}
+	}
+}
+}
+
+void FPersistSandboxView::ShowPersistWorkflow(FPersistSandboxWorkflow& InWorkflow, TConstArrayView<FString> InPreSelectedFiles) const
 {
 	const TSharedRef<SWindow> NewWindow = SNew(SWindow)
 		.Title(LOCTEXT("PersistSubmitWindowTitle", "Persist Sandbox"))
@@ -37,6 +62,7 @@ void FPersistSandboxView::ShowPersistWorkflow(FPersistSandboxWorkflow& InWorkflo
 		.SupportsMinimize(false);
 	
 	const TSharedRef<FPersistOperationViewModel> ViewModel = MakeShared<FPersistOperationViewModel>(InWorkflow);
+	PersistViewDetail::PreSelectFiles(*ViewModel, InPreSelectedFiles);
 	NewWindow->SetContent(
 		SNew(SPersistOperationWidget, ViewModel)
 	);
diff --git a/Engine/Plugins/Developer/Sandbox/SandboxedEditing/Source/SandboxedEditing/Private/Features/Browser/Views/PersistSandboxView.h b/Engine/Plugins/Developer/Sandbox/SandboxedEditing/Source/SandboxedEditing/Private/Features/Browser/Views/PersistSandboxView.h
--- a/Engine/Plugins/Developer/Sandbox/SandboxedEditing/Source/SandboxedEditing/Private/Features/Browser/Views/PersistSandboxView.h
+++ b/Engine/Plugins/Developer/Sandbox/SandboxedEditing/Source/SandboxedEditing/Private/Features/Browser/Views/PersistSandboxView.h
@@ -24,7 +24,7 @@ private:
 	const TSharedRef<FPersistSandboxViewModel> PersistViewModel;
 	
 	/** Shows a modal in which the user can select the files to persist. */
-	void ShowPersistWorkflow(FPersistSandboxWorkflow& InWorkflow) const;
+	void ShowPersistWorkflow(FPersistSandboxWorkflow& InWorkflow, TConstArrayView<FString> InPreSelectedFiles) const;
 	/** Shows a notification summarizing the result of the operation. */
 	void ShowPersistSummary(const FileSandboxUI::FPersistSummary& InSummary) const;
 };
diff --git a/Engine/Plugins/Developer/Sandbox/SandboxedEditing/Source/SandboxedEditing/Private/Features/Browser/Widgets/SSandboxedEditingRoot.cpp b/Engine/Plugins/Developer/Sandbox/SandboxedEditing/Source/SandboxedEditing/Private/Features/Browser/Widgets/SSandboxedEditingRoot.cpp
--- a/Engine/Plugins/Developer/Sandbox/SandboxedEditing/Source/SandboxedEditing/Private/Features/Browser/Widgets/SSandboxedEditingRoot.cpp
+++ b/Engine/Plugins/Developer/Sandbox/SandboxedEditing/Source/SandboxedEditing/Private/Features/Browser/Widgets/SSandboxedEditingRoot.cpp
@@ -21,7 +21,8 @@ void SSandboxedEditingRoot::Construct(const FArguments& InArgs, const FBrowserVi
 	ExternalSandboxViewModel = InViewModels.ExternalSandboxViewModel;
 	ActiveSandboxViewModels.Emplace(
 		InViewModels.ActiveSandboxDetailsViewModel, InViewModels.MetaDataViewModel, 
-		InViewModels.ActiveSandboxFileStateViewModel, InViewModels.FilterActiveSandboxFileStateViewModel
+		InViewModels.ActiveSandboxFileStateViewModel, InViewModels.FilterActiveSandboxFileStateViewModel,
+		InViewModels.PersistViewModel
 		);
 	ActiveFileStateColumnsAttr = InArgs._ActiveFileStateColumnsFactories;
 	
diff --git a/Engine/Plugins/Developer/Sandbox/SandboxedEditing/Source/SandboxedEditing/Private/Features/Browser/Widgets/Sandboxed/SActiveSandbox.cpp b/Engine/Plugins/Developer/Sandbox/SandboxedEditing/Source/SandboxedEditing/Private/Features/Browser/Widgets/Sandboxed/SActiveSandbox.cpp
--- a/Engine/Plugins/Developer/Sandbox/SandboxedEditing/Source/SandboxedEditing/Private/Features/Browser/Widgets/Sandboxed/SActiveSandbox.cpp
+++ b/Engine/Plugins/Developer/Sandbox/SandboxedEditing/Source/SandboxedEditing/Private/Features/Browser/Widgets/Sandboxed/SActiveSandbox.cpp
@@ -91,7 +91,8 @@ TSharedRef<SWidget> SActiveSandbox::MakeDescription(const FViewModels& InViewMod
 TSharedRef<SWidget> SActiveSandbox::MakeFileChangesContent(const FArguments& InArgs, const FViewModels& InViewModels)
 {
 	const SActiveSandboxFileChanges::FViewModels ViewModels(
-		InViewModels.ActiveSandboxViewModel, InViewModels.ActiveSandboxFileStateModel, InViewModels.FilterActiveSandboxFileStateModel
+		InViewModels.ActiveSandboxViewModel, InViewModels.ActiveSandboxFileStateModel, InViewModels.FilterActiveSandboxFileStateModel,
+		InViewModels.PersistViewModel
 		);
 	
 	return SNew(SExpandableArea)
diff --git a/Engine/Plugins/Developer/Sandbox/SandboxedEditing/Source/SandboxedEditing/Private/Features/Browser/Widgets/Sandboxed/SActiveSandbox.h b/Engine/Plugins/Developer/Sandbox/SandboxedEditing/Source/SandboxedEditing/Private/Features/Browser/Widgets/Sandboxed/SActiveSandbox.h
--- a/Engine/Plugins/Developer/Sandbox/SandboxedEditing/Source/SandboxedEditing/Private/Features/Browser/Widgets/Sandboxed/SActiveSandbox.h
+++ b/Engine/Plugins/Developer/Sandbox/SandboxedEditing/Source/SandboxedEditing/Private/Features/Browser/Widgets/Sandboxed/SActiveSandbox.h
@@ -13,6 +13,7 @@ namespace UE::SandboxedEditing
 class FActiveSandboxDetailsViewModel;
 class FActiveSandboxFileStateViewModel;
 class FFilterFileStateViewModel;
+class FPersistSandboxViewModel;
 class FSandboxMetaDataViewModel;
 class SFilterableFileStateListView;
 
@@ -27,17 +28,20 @@ public:
 		TSharedRef<FSandboxMetaDataViewModel> MetaDataViewModel;
 		TSharedRef<FActiveSandboxFileStateViewModel> ActiveSandboxFileStateModel;
 		TSharedRef<FFilterFileStateViewModel> FilterActiveSandboxFileStateModel;
+		TSharedRef<FPersistSandboxViewModel> PersistViewModel;
 
 		explicit FViewModels(
 			const TSharedRef<FActiveSandboxDetailsViewModel>& InActiveSandboxViewModel,
 			const TSharedRef<FSandboxMetaDataViewModel>& InMetaDataViewModel,
 			const TSharedRef<FActiveSandboxFileStateViewModel>& InActiveSandboxFileStateModel,
-			const TSharedRef<FFilterFileStateViewModel>& InFilterActiveSandboxFileStateModel
+			const TSharedRef<FFilterFileStateViewModel>& InFilterActiveSandboxFileStateModel,
+			const TSharedRef<FPersistSandboxViewModel>& InPersistViewModel
 			)
 			: ActiveSandboxViewModel(InActiveSandboxViewModel)
 			, MetaDataViewModel(InMetaDataViewModel)
 			, ActiveSandboxFileStateModel(InActiveSandboxFileStateModel)
 			, FilterActiveSandboxFileStateModel(InFilterActiveSandboxFileStateModel)
+			, PersistViewModel(InPersistViewModel)
 		{}
 	};
 	
diff --git a/Engine/Plugins/Developer/Sandbox/SandboxedEditing/Source/SandboxedEditing/Private/Features/Browser/Widgets/Sandboxed/SActiveSandboxFileChanges.cpp b/Engine/Plugins/Developer/Sandbox/SandboxedEditing/Source/SandboxedEditing/Private/Features/Browser/Widgets/Sandboxed/SActiveSandboxFileChanges.cpp
--- a/Engine/Plugins/Developer/Sandbox/SandboxedEditing/Source/SandboxedEditing/Private/Features/Browser/Widgets/Sandboxed/SActiveSandboxFileChanges.cpp
+++ b/Engine/Plugins/Developer/Sandbox/SandboxedEditing/Source/SandboxedEditing/Private/Features/Browser/Widgets/Sandboxed/SActiveSandboxFileChanges.cpp
@@ -2,8 +2,10 @@
 
 #include "SActiveSandboxFileChanges.h"
 
+#include "Features/Browser/Commands/FileStateActions/FileStateCommands.h"
 #include "Features/Browser/ViewModels/Active/ActiveSandboxDetailsViewModel.h"
 #include "Features/Browser/ViewModels/FileState/FileStateItem.h"
+#include "Features/Browser/ViewModels/Persist/PersistSandboxViewModel.h"
 #include "Features/Browser/Views/Menu/FileStateMenuUtils.h"
 #include "Features/Browser/Widgets/Shared/FileState/SFileStateListView.h"
 #include "Features/Browser/Widgets/Shared/FileState/SFilterableFileStateListView.h"
@@ -17,7 +19,8 @@ void SActiveSandboxFileChanges::Construct(
 	const FArguments& InArgs, const FViewModels& InViewModels
 	)
 {
-	ActiveSandboxViewModel = InViewModels.ActiveSandboxViewModel;
+	PersistViewModel = InViewModels.PersistViewModel;
+	CommandList = InArgs._CommandList;
 	
 	ChildSlot
 	[
@@ -27,20 +30,57 @@ void SActiveSandboxFileChanges::Construct(
 		.OnContextMenuOpening(this, &SActiveSandboxFileChanges::MakeFileChangeContextMenu)	
 	];
 	
-	FileStateCommandBindings = MakeUnique<FFileStateCommandBindings>(
-		ActiveSandboxViewModel->GetSandboxPath(), 
-		FileStateMenu::MakeSelectedFilesAttribute(FileActionWidget->GetListView().ToSharedRef())
-		);
+	BindCommands(InViewModels);
 }
 
 TSharedPtr<SWidget> SActiveSandboxFileChanges::MakeFileChangeContextMenu()
 {
-	FMenuBuilder Menu(true, nullptr);
+	FMenuBuilder Menu(true, CommandList);
+	
+	FFileStateCommands& FileStateCommands = FFileStateCommands::Get();
+	
+	Menu.BeginSection(TEXT("Actions"), LOCTEXT("FileActions", "File Actions"));
+	Menu.AddMenuEntry(FileStateCommands.PersistSelected);
+	Menu.EndSection();
 	
 	FileStateCommandBindings->AppendMenu(Menu);
 	
 	return Menu.MakeWidget();
 }
+
+void SActiveSandboxFileChanges::BindCommands(const FViewModels& InViewModels)
+{
+	FileStateCommandBindings = MakeUnique<FFileStateCommandBindings>(
+		InViewModels.ActiveSandboxViewModel->GetSandboxPath(), 
+		FileStateMenu::MakeSelectedFilesAttribute(FileActionWidget->GetListView().ToSharedRef())
+		);
+	
+	if (!CommandList)
+	{
+		return;
+	}
+	
+	FFileStateCommands& FileStateCommands = FFileStateCommands::Get();
+	CommandList->MapAction(FileStateCommands.PersistSelected,
+		FExecuteAction::CreateSP(this, &SActiveSandboxFileChanges::HandlePersistSelected),
+		FCanExecuteAction::CreateSP(this, &SActiveSandboxFileChanges::CanPersistSelected)
+		);
+}
+
+void SActiveSandboxFileChanges::HandlePersistSelected() const
+{
+	if (CanPersistSelected())
+	{
+		PersistViewModel->StartPersistWorkflowForActiveSandbox(
+			FileStateMenu::TransformSelectedItems(FileActionWidget->GetListView().ToSharedRef())
+			);
+	}
+}
+
+bool SActiveSandboxFileChanges::CanPersistSelected() const
+{
+	return PersistViewModel->CanStartPersistWorkflowForActiveSandbox();
+}
 }
 
 #undef LOCTEXT_NAMESPACE
\ No newline at end of file
diff --git a/Engine/Plugins/Developer/Sandbox/SandboxedEditing/Source/SandboxedEditing/Private/Features/Browser/Widgets/Sandboxed/SActiveSandboxFileChanges.h b/Engine/Plugins/Developer/Sandbox/SandboxedEditing/Source/SandboxedEditing/Private/Features/Browser/Widgets/Sandboxed/SActiveSandboxFileChanges.h
--- a/Engine/Plugins/Developer/Sandbox/SandboxedEditing/Source/SandboxedEditing/Private/Features/Browser/Widgets/Sandboxed/SActiveSandboxFileChanges.h
+++ b/Engine/Plugins/Developer/Sandbox/SandboxedEditing/Source/SandboxedEditing/Private/Features/Browser/Widgets/Sandboxed/SActiveSandboxFileChanges.h
@@ -11,6 +11,7 @@ class FUICommandList;
 
 namespace UE::SandboxedEditing
 {
+class FPersistSandboxViewModel;
 class FActiveSandboxDetailsViewModel;
 class FFilterFileStateViewModel;
 class IFileStateViewModel;
@@ -26,15 +27,18 @@ public:
 		TSharedRef<FActiveSandboxDetailsViewModel> ActiveSandboxViewModel;
 		TSharedRef<IFileStateViewModel> FileActionsViewModel;
 		TSharedRef<FFilterFileStateViewModel> FilterViewModel;
+		TSharedRef<FPersistSandboxViewModel> PersistViewModel;
 
 		explicit FViewModels(
-			const TSharedRef<FActiveSandboxDetailsViewModel>& ActiveSandboxViewModel,
-			const TSharedRef<IFileStateViewModel>& FileActionsViewModel,
-			const TSharedRef<FFilterFileStateViewModel>& FilterViewModel
+			const TSharedRef<FActiveSandboxDetailsViewModel>& InActiveSandboxViewModel,
+			const TSharedRef<IFileStateViewModel>& InFileActionsViewModel,
+			const TSharedRef<FFilterFileStateViewModel>& InFilterViewModel,
+			const TSharedRef<FPersistSandboxViewModel>& InPersistViewModel
 			)
-			: ActiveSandboxViewModel(ActiveSandboxViewModel)
-			, FileActionsViewModel(FileActionsViewModel)
-			, FilterViewModel(FilterViewModel)
+			: ActiveSandboxViewModel(InActiveSandboxViewModel)
+			, FileActionsViewModel(InFileActionsViewModel)
+			, FilterViewModel(InFilterViewModel)
+			, PersistViewModel(InPersistViewModel)
 		{}
 	};
 	
@@ -50,8 +54,10 @@ public:
 	
 private:
 	
-	/** Used for binding commands. */
-	TSharedPtr<FActiveSandboxDetailsViewModel> ActiveSandboxViewModel;
+	/** Used for building the right-click menu. */
+	TSharedPtr<FUICommandList> CommandList;
+	/** Used for binding right-click persist commands. */
+	TSharedPtr<FPersistSandboxViewModel> PersistViewModel;
 	
 	/** Displays the file changes made in this sandbox. */
 	TSharedPtr<SFilterableFileStateListView> FileActionWidget;
@@ -61,6 +67,14 @@ private:
 	
 	/** @return The context menu for FileActionWidget. */
 	TSharedPtr<SWidget> MakeFileChangeContextMenu();
+
+	/** Binds the commands that this widget introduces. */
+	void BindCommands(const FViewModels& InViewModels);
+	
+	/** Brings up the persist dialogue with the current selection pre-selected. */
+	void HandlePersistSelected() const;
+	/** @return Whether the persist dialogue can be created at this time */
+	bool CanPersistSelected() const;
 };
 }
 
