Sandboxed Editing: Add settings button on browser toolbar which opens the plugin settings.

#Synced-CL 51334859

diff --git a/Engine/Plugins/Developer/Sandbox/SandboxedEditing/Source/SandboxedEditing/Private/Features/Browser/Widgets/Sandboxed/SActiveSandboxToolbar.cpp b/Engine/Plugins/Developer/Sandbox/SandboxedEditing/Source/SandboxedEditing/Private/Features/Browser/Widgets/Sandboxed/SActiveSandboxToolbar.cpp
--- a/Engine/Plugins/Developer/Sandbox/SandboxedEditing/Source/SandboxedEditing/Private/Features/Browser/Widgets/Sandboxed/SActiveSandboxToolbar.cpp
+++ b/Engine/Plugins/Developer/Sandbox/SandboxedEditing/Source/SandboxedEditing/Private/Features/Browser/Widgets/Sandboxed/SActiveSandboxToolbar.cpp
@@ -72,6 +72,15 @@ void SActiveSandboxToolbar::Construct(const FArguments& InArgs, const TSharedRef
 				TAttribute<FText>::CreateSP(this, &SActiveSandboxToolbar::GetLeaveButtonToolTipText)
 				)
 		]
+
+		// Go to settings
+		+SHorizontalBox::Slot()
+		.AutoWidth()
+		.VAlign(VAlign_Center)
+		.Padding(FMargin(1.0f))
+		[
+			CreateOpenSettingsButton()
+		]
 	];
 }
 
diff --git a/Engine/Plugins/Developer/Sandbox/SandboxedEditing/Source/SandboxedEditing/Private/Features/Browser/Widgets/Unsandboxed/SBrowserControls.cpp b/Engine/Plugins/Developer/Sandbox/SandboxedEditing/Source/SandboxedEditing/Private/Features/Browser/Widgets/Unsandboxed/SBrowserControls.cpp
--- a/Engine/Plugins/Developer/Sandbox/SandboxedEditing/Source/SandboxedEditing/Private/Features/Browser/Widgets/Unsandboxed/SBrowserControls.cpp
+++ b/Engine/Plugins/Developer/Sandbox/SandboxedEditing/Source/SandboxedEditing/Private/Features/Browser/Widgets/Unsandboxed/SBrowserControls.cpp
@@ -2,14 +2,19 @@
 
 #include "SBrowserControls.h"
 
+#include "ISettingsModule.h"
 #include "SandboxedEditingStyle.h"
 #include "Features/Browser/Commands/BrowserCommands.h"
 #include "Features/Browser/ViewModels/SandboxControlsViewModel.h"
 #include "Filters/CustomTextFilters.h"
 #include "Framework/Commands/GenericCommands.h"
 #include "Framework/MultiBox/MultiBoxBuilder.h"
+#include "Modules/ModuleManager.h"
 #include "Utils/WidgetUtils.h"
 #include "Widgets/SBoxPanel.h"
+#include "Widgets/Images/SImage.h"
+#include "Widgets/Input/SButton.h"
+#include "Widgets/Layout/SSpacer.h"
 
 #define LOCTEXT_NAMESPACE "SBrowserControls"
 
@@ -21,7 +26,28 @@ void SBrowserControls::Construct(const FArguments& InArgs, const TSharedRef<FSan
 	
 	ChildSlot
 	[
-		BuildButtonBar(InArgs)
+		SNew(SHorizontalBox)
+
+		+SHorizontalBox::Slot()
+		.AutoWidth()
+		.VAlign(VAlign_Center)
+		[
+			BuildButtonBar(InArgs)
+		]
+
+		+SHorizontalBox::Slot()
+		.FillWidth(1.f)
+		[
+			SNew(SSpacer)
+		]
+		
+		+SHorizontalBox::Slot()
+		.AutoWidth()
+		.VAlign(VAlign_Center)
+		.Padding(1.f)
+		[
+			CreateOpenSettingsButton()
+		]
 	];
 }
 
diff --git a/Engine/Plugins/Developer/Sandbox/SandboxedEditing/Source/SandboxedEditing/Private/Utils/WidgetUtils.cpp b/Engine/Plugins/Developer/Sandbox/SandboxedEditing/Source/SandboxedEditing/Private/Utils/WidgetUtils.cpp
--- a/Engine/Plugins/Developer/Sandbox/SandboxedEditing/Source/SandboxedEditing/Private/Utils/WidgetUtils.cpp
+++ b/Engine/Plugins/Developer/Sandbox/SandboxedEditing/Source/SandboxedEditing/Private/Utils/WidgetUtils.cpp
@@ -2,14 +2,16 @@
 
 #include "WidgetUtils.h"
 
-#include "SNegativeActionButton.h"
-#include "SPositiveActionButton.h"
-#include "SPrimaryButton.h"
-#include "Framework/SlateDelegates.h"
 #include "Framework/Commands/UICommandInfo.h"
 #include "Framework/Commands/UICommandList.h"
+#include "Framework/SlateDelegates.h"
+#include "ISettingsModule.h"
 #include "Internationalization/Text.h"
 #include "Misc/Attribute.h"
+#include "Modules/ModuleManager.h"
+#include "SNegativeActionButton.h"
+#include "SPositiveActionButton.h"
+#include "SPrimaryButton.h"
 #include "Styling/AppStyle.h"
 #include "Styling/SlateColor.h"
 #include "Widgets/Images/SImage.h"
@@ -147,4 +149,21 @@ TSharedRef<SWidget> MakeNegativeActionButton(
 		.IsEnabled(EnabledAttribute)
 		.Icon(Icon);
 }
-}
\ No newline at end of file
+
+TSharedRef<SWidget> CreateOpenSettingsButton()
+{
+	return SNew(SButton)
+		.ButtonStyle(FAppStyle::Get(), TEXT("SimpleButton"))
+		.OnClicked_Lambda([]()
+		{
+			FModuleManager::GetModulePtr<ISettingsModule>("Settings")->ShowViewer("Project", "Plugins", "Sandbox"); 
+			return FReply::Handled();
+		})
+		.VAlign(VAlign_Center)
+		.HAlign(HAlign_Fill)
+		[
+			SNew(SImage)
+			.Image(FAppStyle::GetBrush("Icons.Settings"))
+		];
+}
+}
diff --git a/Engine/Plugins/Developer/Sandbox/SandboxedEditing/Source/SandboxedEditing/Private/Utils/WidgetUtils.h b/Engine/Plugins/Developer/Sandbox/SandboxedEditing/Source/SandboxedEditing/Private/Utils/WidgetUtils.h
--- a/Engine/Plugins/Developer/Sandbox/SandboxedEditing/Source/SandboxedEditing/Private/Utils/WidgetUtils.h
+++ b/Engine/Plugins/Developer/Sandbox/SandboxedEditing/Source/SandboxedEditing/Private/Utils/WidgetUtils.h
@@ -71,4 +71,7 @@ TSharedRef<SWidget> MakeNegativeActionButton(
 	const FOnClicked& OnClicked, 
 	const TAttribute<EVisibility>& Visibility = EVisibility::Visible
 	);
+
+/** @return A gear button that opens the plugin settings for Sandboxed Editing. */
+TSharedRef<SWidget> CreateOpenSettingsButton();
 }
