Sandboxed Editing: Add "Revert" entry to right-click context menu of the file changes list view while sandboxed.
This is guarded by the new "SandboxedEditing.AllowRevertUI" CVar to restrict user access. The intention is to allow developers to test this functionality before it is made non-experimental.

#Synced-CL 51334859

diff --git a/Engine/Plugins/Developer/Sandbox/SandboxedEditing/Source/SandboxedEditing/Private/Features/Browser/Widgets/Sandboxed/SActiveSandboxFileChanges.cpp b/Engine/Plugins/Developer/Sandbox/SandboxedEditing/Source/SandboxedEditing/Private/Features/Browser/Widgets/Sandboxed/SActiveSandboxFileChanges.cpp
--- a/Engine/Plugins/Developer/Sandbox/SandboxedEditing/Source/SandboxedEditing/Private/Features/Browser/Widgets/Sandboxed/SActiveSandboxFileChanges.cpp
+++ b/Engine/Plugins/Developer/Sandbox/SandboxedEditing/Source/SandboxedEditing/Private/Features/Browser/Widgets/Sandboxed/SActiveSandboxFileChanges.cpp
@@ -9,12 +9,22 @@
 #include "Features/Browser/Views/Menu/FileStateMenuUtils.h"
 #include "Features/Browser/Widgets/Shared/FileState/SFileStateListView.h"
 #include "Features/Browser/Widgets/Shared/FileState/SFilterableFileStateListView.h"
+#include "Framework/Models/SandboxSystemModel.h"
 #include "Framework/MultiBox/MultiBoxBuilder.h"
 
 #define LOCTEXT_NAMESPACE "SActiveSandboxFileChanges"
 
 namespace UE::SandboxedEditing
 {
+namespace ActiveFileChangesDetail
+{
+/** 
+ * TODO DP: This feature is experimental and not fully tested yet. What if you revert a change that another file depends on?
+ * E.g. Revert adding a blueprint but a map has that BP saved inside of it?
+ */
+static TAutoConsoleVariable<bool> CVarAllowRevertUI(TEXT("SandboxedEditing.AllowRevertUI"), false, TEXT("Whether to allow the user to revert files."));
+}
+
 void SActiveSandboxFileChanges::Construct(
 	const FArguments& InArgs, const FViewModels& InViewModels
 	)
@@ -36,11 +46,21 @@ void SActiveSandboxFileChanges::Construct(
 TSharedPtr<SWidget> SActiveSandboxFileChanges::MakeFileChangeContextMenu()
 {
 	FMenuBuilder Menu(true, CommandList);
-	
 	FFileStateCommands& FileStateCommands = FFileStateCommands::Get();
-	
+
 	Menu.BeginSection(TEXT("Actions"), LOCTEXT("FileActions", "File Actions"));
 	Menu.AddMenuEntry(FileStateCommands.PersistSelected);
+	if (ActiveFileChangesDetail::CVarAllowRevertUI.GetValueOnGameThread())
+	{
+		Menu.AddMenuEntry(
+			LOCTEXT("Revert", "Revert"), 
+			FText::GetEmpty(),
+			FSlateIcon(),
+			FUIAction(
+				FExecuteAction::CreateSP(this, &SActiveSandboxFileChanges::RevertSelection),
+				FCanExecuteAction::CreateSP(this, &SActiveSandboxFileChanges::CanRevertSelection)
+			));
+	}
 	Menu.EndSection();
 	
 	FileStateCommandBindings->AppendMenu(Menu);
@@ -81,6 +101,21 @@ bool SActiveSandboxFileChanges::CanPersistSelected() const
 {
 	return PersistViewModel->CanStartPersistWorkflowForActiveSandbox();
 }
+
+void SActiveSandboxFileChanges::RevertSelection() const
+{
+	if (CanRevertSelection())
+	{
+		const TArray<FString> Selection = FileStateMenu::TransformSelectedItems(FileActionWidget->GetListView().ToSharedRef());
+		SandboxModel::RevertSpecifiedChanges(Selection);
+	}
+}
+
+bool SActiveSandboxFileChanges::CanRevertSelection() const
+{
+	const TArray<FString> Selection = FileStateMenu::TransformSelectedItems(FileActionWidget->GetListView().ToSharedRef());
+	return !Selection.IsEmpty() && SandboxModel::CanRevertSpecifiedChanges(Selection);
+}
 }
 
 #undef LOCTEXT_NAMESPACE
\ No newline at end of file
diff --git a/Engine/Plugins/Developer/Sandbox/SandboxedEditing/Source/SandboxedEditing/Private/Features/Browser/Widgets/Sandboxed/SActiveSandboxFileChanges.h b/Engine/Plugins/Developer/Sandbox/SandboxedEditing/Source/SandboxedEditing/Private/Features/Browser/Widgets/Sandboxed/SActiveSandboxFileChanges.h
--- a/Engine/Plugins/Developer/Sandbox/SandboxedEditing/Source/SandboxedEditing/Private/Features/Browser/Widgets/Sandboxed/SActiveSandboxFileChanges.h
+++ b/Engine/Plugins/Developer/Sandbox/SandboxedEditing/Source/SandboxedEditing/Private/Features/Browser/Widgets/Sandboxed/SActiveSandboxFileChanges.h
@@ -64,7 +64,7 @@ private:
 	
 	/** Common entries for FileActionWidget. */
 	TUniquePtr<FFileStateCommandBindings> FileStateCommandBindings;
-	
+
 	/** @return The context menu for FileActionWidget. */
 	TSharedPtr<SWidget> MakeFileChangeContextMenu();
 
@@ -75,6 +75,11 @@ private:
 	void HandlePersistSelected() const;
 	/** @return Whether the persist dialogue can be created at this time */
 	bool CanPersistSelected() const;
+	
+	/** Handles command for reverting current selection. */
+	void RevertSelection() const;
+	/** Whether there is a selection, and it can be reverted. */
+	bool CanRevertSelection() const;
 };
 }
 
diff --git a/Engine/Plugins/Developer/Sandbox/SandboxedEditing/Source/SandboxedEditing/Private/Features/StatusBar/Commands/StatusBarCommandMappings.cpp b/Engine/Plugins/Developer/Sandbox/SandboxedEditing/Source/SandboxedEditing/Private/Features/StatusBar/Commands/StatusBarCommandMappings.cpp
--- a/Engine/Plugins/Developer/Sandbox/SandboxedEditing/Source/SandboxedEditing/Private/Features/StatusBar/Commands/StatusBarCommandMappings.cpp
+++ b/Engine/Plugins/Developer/Sandbox/SandboxedEditing/Source/SandboxedEditing/Private/Features/StatusBar/Commands/StatusBarCommandMappings.cpp
@@ -39,7 +39,7 @@ FStatusBarCommandMappings::FStatusBarCommandMappings(const TSharedRef<FSandboxSy
 	CommandList->MapAction(
 		Commands.DiscardAll,
 		FExecuteAction::CreateRaw(this, &FStatusBarCommandMappings::HandleDiscardAll),
-		FCanExecuteAction::CreateSP(InSandboxViewModel, &FSandboxSystemModel::CanDiscardAllChanges)
+		FCanExecuteAction::CreateSP(InSandboxViewModel, &FSandboxSystemModel::CanRevertAllChanges)
 	);
 }
 
@@ -66,6 +66,6 @@ void FStatusBarCommandMappings::HandlePersistAll()
 
 void FStatusBarCommandMappings::HandleDiscardAll()
 {
-	SandboxViewModel->DiscardAllChanges();
+	SandboxViewModel->RevertAllChanges();
 }
 }
diff --git a/Engine/Plugins/Developer/Sandbox/SandboxedEditing/Source/SandboxedEditing/Private/Framework/Models/SandboxSystemModel.cpp b/Engine/Plugins/Developer/Sandbox/SandboxedEditing/Source/SandboxedEditing/Private/Framework/Models/SandboxSystemModel.cpp
--- a/Engine/Plugins/Developer/Sandbox/SandboxedEditing/Source/SandboxedEditing/Private/Framework/Models/SandboxSystemModel.cpp
+++ b/Engine/Plugins/Developer/Sandbox/SandboxedEditing/Source/SandboxedEditing/Private/Framework/Models/SandboxSystemModel.cpp
@@ -309,20 +309,20 @@ void FSandboxSystemModel::EnumerateFileChanges(
 	SandboxManager.EnumerateFileChanges(InSandbox, InCallback);
 }
 
-void FSandboxSystemModel::DiscardAllChanges() const
+void FSandboxSystemModel::RevertAllChanges() const
 {
 	using namespace FileSandboxCore;
 	IFileSandboxCoreModule& Module = IFileSandboxCoreModule::Get();
 	ISandboxManager& SandboxManager = Module.GetSandboxManager();
 	ISandboxInstance* Instance = SandboxManager.GetActiveSandboxInstance();
 	
-	if (CanDiscardAllChanges() && ensure(Instance))
+	if (CanRevertAllChanges() && ensure(Instance))
 	{
 		Instance->RevertAll();
 	}
 }
 
-bool FSandboxSystemModel::CanDiscardAllChanges() const
+bool FSandboxSystemModel::CanRevertAllChanges() const
 {
 	return HasActiveSandbox();
 }
@@ -493,6 +493,32 @@ FString GetBaseSandboxDirectory()
 	// For now we use the global default directory but this could change in the future, e.g. by settings.
 	return FileSandboxCore::GetBaseSandboxDirectory();
 }
+
+bool HasActiveSandbox()
+{
+	using namespace FileSandboxCore;
+	IFileSandboxCoreModule& Module = IFileSandboxCoreModule::Get();
+	ISandboxInstance* Instance = Module.GetSandboxManager().GetActiveSandboxInstance();
+	return Instance != nullptr;
+}
+
+void RevertSpecifiedChanges(TConstArrayView<FString> InFiles)
+{
+	using namespace FileSandboxCore;
+	IFileSandboxCoreModule& Module = IFileSandboxCoreModule::Get();
+	ISandboxManager& SandboxManager = Module.GetSandboxManager();
+	ISandboxInstance* Instance = SandboxManager.GetActiveSandboxInstance();
+	
+	if (CanRevertSpecifiedChanges(InFiles) && ensure(Instance))
+	{
+		Instance->RevertSpecified(InFiles);
+	}
+}
+
+bool CanRevertSpecifiedChanges(TConstArrayView<FString> InFiles)
+{
+	return HasActiveSandbox() && !InFiles.IsEmpty();
+}
 }
 }
 
diff --git a/Engine/Plugins/Developer/Sandbox/SandboxedEditing/Source/SandboxedEditing/Private/Framework/Models/SandboxSystemModel.h b/Engine/Plugins/Developer/Sandbox/SandboxedEditing/Source/SandboxedEditing/Private/Framework/Models/SandboxSystemModel.h
--- a/Engine/Plugins/Developer/Sandbox/SandboxedEditing/Source/SandboxedEditing/Private/Framework/Models/SandboxSystemModel.h
+++ b/Engine/Plugins/Developer/Sandbox/SandboxedEditing/Source/SandboxedEditing/Private/Framework/Models/SandboxSystemModel.h
@@ -89,10 +89,10 @@ public:
 		TFunctionRef<FileSandboxCore::EBreakBehavior(const FileSandboxCore::FSandboxedFileChangeInfo& InChange)> InCallback
 		);
 
-	/** Discards all changes */
-	void DiscardAllChanges() const;
-	/** @return Whether all changes can be discarded */
-	bool CanDiscardAllChanges() const; 
+	/** Reverts all changes */
+	void RevertAllChanges() const;
+	/** @return Whether all changes can be reverted */
+	bool CanRevertAllChanges() const; 
 
 	/** @return Path to the root directory of the active sandbox. Empty if none is active. */
 	FString GetActiveSandboxPath() const;
@@ -168,5 +168,13 @@ TOptional<FString> GetSandboxName(const FString& InSandboxDirectory);
 
 /** @return The base directory in which sandboxes are placed. */
 FString GetBaseSandboxDirectory();
+
+/** @return Whether there is an active sandbox */
+bool HasActiveSandbox();
+
+/** Reverts the changes of the specified files. */
+void RevertSpecifiedChanges(TConstArrayView<FString> InFiles);
+/** @return Whether the specified files can be reverted. */
+bool CanRevertSpecifiedChanges(TConstArrayView<FString> InFiles);
 }
 }
\ No newline at end of file
