SandboxedEditing - Fix extra qualification warnings when compiing with ClangInstrument

#Synced-CL 52068815

diff --git a/Engine/Plugins/Developer/Sandbox/SandboxedEditing/Source/SandboxedEditing/Private/Features/Browser/Views/Menu/FileStateMenuUtils.cpp b/Engine/Plugins/Developer/Sandbox/SandboxedEditing/Source/SandboxedEditing/Private/Features/Browser/Views/Menu/FileStateMenuUtils.cpp
--- a/Engine/Plugins/Developer/Sandbox/SandboxedEditing/Source/SandboxedEditing/Private/Features/Browser/Views/Menu/FileStateMenuUtils.cpp
+++ b/Engine/Plugins/Developer/Sandbox/SandboxedEditing/Source/SandboxedEditing/Private/Features/Browser/Views/Menu/FileStateMenuUtils.cpp
@@ -24,7 +24,7 @@ TArray<FString> TransformSelectedItems(const TSharedRef<SFileStateListView>& InL
 	return NonSandboxFiles;
 }
 
-TAttribute<TArray<FString>> FileStateMenu::MakeSelectedFilesAttribute(const TSharedRef<SFileStateListView>& InListView)
+TAttribute<TArray<FString>> MakeSelectedFilesAttribute(const TSharedRef<SFileStateListView>& InListView)
 {
 	return TAttribute<TArray<FString>>::CreateLambda([WeakListView = InListView.ToWeakPtr()]
 	{
@@ -44,7 +44,7 @@ static FText GetExploreFolderText()
 }
 
 
-void FileStateMenu::AppendMenu(FMenuBuilder& InMenuBuilder)
+void AppendMenu(FMenuBuilder& InMenuBuilder)
 {
 	InMenuBuilder.BeginSection(TEXT("Explorer"), LOCTEXT("Explorer", "Explorer"));
 	
diff --git a/Engine/Plugins/Developer/Sandbox/SandboxedEditing/Source/SandboxedEditing/Private/Framework/Models/SandboxSystemModel.cpp b/Engine/Plugins/Developer/Sandbox/SandboxedEditing/Source/SandboxedEditing/Private/Framework/Models/SandboxSystemModel.cpp
--- a/Engine/Plugins/Developer/Sandbox/SandboxedEditing/Source/SandboxedEditing/Private/Framework/Models/SandboxSystemModel.cpp
+++ b/Engine/Plugins/Developer/Sandbox/SandboxedEditing/Source/SandboxedEditing/Private/Framework/Models/SandboxSystemModel.cpp
@@ -478,7 +478,7 @@ void FSandboxSystemModel::HandleActiveSandboxFilesChanged()
 
 namespace SandboxModel
 {
-TOptional<FString> SandboxModel::GetSandboxPathFor(const FString& InSandbox, const FString& InNonSandboxPath)
+TOptional<FString> GetSandboxPathFor(const FString& InSandbox, const FString& InNonSandboxPath)
 {
 	return FileSandboxCore::GetSandboxPathFor(InSandbox, InNonSandboxPath);
 }
