Sandbox startup behavior work.

#rnx
#Synced-CL 52173051

diff --git a/Engine/Plugins/Developer/Sandbox/SandboxedEditing/SandboxedEditing.uplugin b/Engine/Plugins/Developer/Sandbox/SandboxedEditing/SandboxedEditing.uplugin
--- a/Engine/Plugins/Developer/Sandbox/SandboxedEditing/SandboxedEditing.uplugin
+++ b/Engine/Plugins/Developer/Sandbox/SandboxedEditing/SandboxedEditing.uplugin
@@ -27,6 +27,10 @@
 		{
 			"Name": "FileSandbox",
 			"Enabled": true
+		},
+		{
+			"Name": "NamingTokens",
+			"Enabled": true
 		}
 	]
 }
diff --git a/Engine/Plugins/Developer/Sandbox/SandboxedEditing/Source/SandboxedEditing/Private/Framework/SandboxedEditingApp.cpp b/Engine/Plugins/Developer/Sandbox/SandboxedEditing/Source/SandboxedEditing/Private/Framework/SandboxedEditingApp.cpp
--- a/Engine/Plugins/Developer/Sandbox/SandboxedEditing/Source/SandboxedEditing/Private/Framework/SandboxedEditingApp.cpp
+++ b/Engine/Plugins/Developer/Sandbox/SandboxedEditing/Source/SandboxedEditing/Private/Framework/SandboxedEditingApp.cpp
@@ -5,6 +5,7 @@
 #include "EntryPoint/SandboxedEditingEntryPoint.h"
 #include "EntryPoint/SandboxedEditingSandboxIntegration.h"
 #include "Models/SandboxSystemModel.h"
+#include "StartupBehavior/StartupBehaviorHandler.h"
 
 namespace UE::SandboxedEditing
 {
@@ -13,6 +14,7 @@ FSandboxedEditingApp::FSandboxedEditingApp()
 	, SandboxIntegration(MakeShared<FSandboxedEditingSandboxIntegration>())
 	, StatusBarFeature(SandboxModel)
 	, BrowserFeature(SandboxModel, SandboxIntegration->GetExternalSandboxViewModel())
+	, StartupBehaviorHandler(SandboxModel)
 {
 	SandboxIntegration->GetEntryPoint()->OnRequestSummonUI().AddLambda([this] { BrowserFeature.SummonUI(); });
 }
diff --git a/Engine/Plugins/Developer/Sandbox/SandboxedEditing/Source/SandboxedEditing/Private/Framework/SandboxedEditingApp.h b/Engine/Plugins/Developer/Sandbox/SandboxedEditing/Source/SandboxedEditing/Private/Framework/SandboxedEditingApp.h
--- a/Engine/Plugins/Developer/Sandbox/SandboxedEditing/Source/SandboxedEditing/Private/Framework/SandboxedEditingApp.h
+++ b/Engine/Plugins/Developer/Sandbox/SandboxedEditing/Source/SandboxedEditing/Private/Framework/SandboxedEditingApp.h
@@ -4,6 +4,7 @@
 
 #include "Features/Browser/BrowserFeature.h"
 #include "Features/StatusBar/StatusBarFeature.h"
+#include "Framework/StartupBehavior/StartupBehaviorHandler.h"
 #include "Templates/SharedPointer.h"
 #include "Templates/UnrealTemplate.h"
 
@@ -11,7 +12,8 @@ namespace UE::SandboxedEditing
 {
 class FSandboxSystemModel;
 class FSandboxedEditingSandboxIntegration;
-	
+class FStartupBehaviorHandler;
+
 /**
  * Creates all relevant systems to deliver the sandboxed editing experience in the editor.
  * This is a coordinator: It sets up features, which are isolated areas in the editor that interact with sandboxing.
@@ -19,7 +21,6 @@ class FSandboxedEditingSandboxIntegration;
 class FSandboxedEditingApp : public FNoncopyable
 {
 public:
-
 	FSandboxedEditingApp();
 
 private:
@@ -29,14 +30,16 @@ private:
 	 * Shared pointer in case the UI to avoid crash in case that though some bug the UI outlives this app for whatever reason.
 	 */
 	const TSharedRef<FSandboxSystemModel> SandboxModel;
-	
+
 	/** Responsible for registering Sandboxed Editing with the sandbox system. */
 	const TSharedRef<FSandboxedEditingSandboxIntegration> SandboxIntegration;
-	
+
 	/** Extends the status bar for Sandboxed Editing. */
 	const FStatusBarFeature StatusBarFeature;
 	/** Extends the editor with a nomad tab containing a browser for sandboxes. */
 	const FBrowserFeature BrowserFeature;
+	/** Handles startup behavior (creating/loading sandboxes on editor startup). */
+	const FStartupBehaviorHandler StartupBehaviorHandler;
 };
 }
 
diff --git a/Engine/Plugins/Developer/Sandbox/SandboxedEditing/Source/SandboxedEditing/Private/Framework/StartupBehavior/StartupBehaviorHandler.cpp b/Engine/Plugins/Developer/Sandbox/SandboxedEditing/Source/SandboxedEditing/Private/Framework/StartupBehavior/StartupBehaviorHandler.cpp
new file mode 100644
--- /dev/null
+++ b/Engine/Plugins/Developer/Sandbox/SandboxedEditing/Source/SandboxedEditing/Private/Framework/StartupBehavior/StartupBehaviorHandler.cpp
@@ -0,0 +1,492 @@
+// Copyright Epic Games, Inc. All Rights Reserved.
+
+#include "StartupBehaviorHandler.h"
+
+#include "Editor.h"
+#include "Framework/Application/SlateApplication.h"
+#include "Framework/Models/CreateSandboxArgs.h"
+#include "Framework/Models/SandboxInfo.h"
+#include "Framework/Models/SandboxSystemModel.h"
+#include "Framework/Notifications/NotificationManager.h"
+#include "LogSandboxedEditing.h"
+#include "Misc/App.h"
+#include "Misc/CommandLine.h"
+#include "Misc/DateTime.h"
+#include "Misc/Parse.h"
+#include "NamingTokensEngineSubsystem.h"
+#include "SandboxedEditingSettings.h"
+#include "Widgets/Notifications/SNotificationList.h"
+
+#define LOCTEXT_NAMESPACE "SandboxedEditing.StartupBehavior"
+
+namespace UE::SandboxedEditing
+{
+
+FStartupBehaviorHandler::FStartupBehaviorHandler(const TSharedRef<FSandboxSystemModel>& InModel)
+	: SandboxModel(InModel)
+{
+	// Register for editor initialization - this is when we execute startup behavior
+	FEditorDelegates::OnEditorInitialized.AddRaw(this, &FStartupBehaviorHandler::OnEditorInitialized);
+
+	// Track sandbox loads for last used tracking
+	LoadSandboxHandle = SandboxModel->OnLoadSandbox().AddRaw(this, &FStartupBehaviorHandler::OnSandboxLoaded);
+
+	UE_LOG(LogSandboxedEditing, Verbose, TEXT("StartupBehaviorHandler initialized"));
+}
+
+FStartupBehaviorHandler::~FStartupBehaviorHandler()
+{
+	// Dismiss any active notification and disable buttons to prevent callbacks during fade
+	if (CurrentNotification.IsValid())
+	{
+		// SetEnabled(false) prevents button clicks during fade animation, ensuring callbacks cannot fire
+		CurrentNotification->SetEnabled(false);
+		CurrentNotification->SetCompletionState(SNotificationItem::CS_None);
+		CurrentNotification->ExpireAndFadeout();
+		CurrentNotification.Reset();
+	}
+
+	FEditorDelegates::OnEditorInitialized.RemoveAll(this);
+
+	if (LoadSandboxHandle.IsValid())
+	{
+		SandboxModel->OnLoadSandbox().Remove(LoadSandboxHandle);
+	}
+
+	UE_LOG(LogSandboxedEditing, Verbose, TEXT("StartupBehaviorHandler destroyed"));
+}
+
+void FStartupBehaviorHandler::OnEditorInitialized(double Duration)
+{
+	UE_LOG(LogSandboxedEditing, Log, TEXT("Editor initialized in %.2f seconds, checking sandbox startup behavior"), Duration);
+
+	// Check if another plugin already started a sandbox
+	if (SandboxModel->HasActiveSandbox())
+	{
+		UE_LOG(LogSandboxedEditing, Log, TEXT("Sandbox already active (%s), skipping startup behavior"),
+			*SandboxModel->GetActiveSandboxName());
+		return;
+	}
+
+	ExecuteStartupBehavior();
+}
+
+bool FStartupBehaviorHandler::TryGetStartupBehaviorFromCommandLine(ESandboxStartupBehavior& OutBehavior) const
+{
+	FString BehaviorString;
+	if (!FParse::Value(FCommandLine::Get(), TEXT("SANDBOXSTARTUP="), BehaviorString))
+	{
+		return false;
+	}
+
+	BehaviorString.TrimStartAndEndInline();
+
+	if (BehaviorString.Equals(TEXT("DoNothing"), ESearchCase::IgnoreCase))
+	{
+		OutBehavior = ESandboxStartupBehavior::DoNothing;
+		return true;
+	}
+	else if (BehaviorString.Equals(TEXT("PromptToJoin"), ESearchCase::IgnoreCase))
+	{
+		OutBehavior = ESandboxStartupBehavior::PromptToJoin;
+		return true;
+	}
+	else if (BehaviorString.Equals(TEXT("CreateNewSandbox"), ESearchCase::IgnoreCase))
+	{
+		OutBehavior = ESandboxStartupBehavior::CreateNewSandbox;
+		return true;
+	}
+	else if (BehaviorString.Equals(TEXT("JoinPreviousSandbox"), ESearchCase::IgnoreCase))
+	{
+		OutBehavior = ESandboxStartupBehavior::JoinPreviousSandbox;
+		return true;
+	}
+	else
+	{
+		UE_LOG(LogSandboxedEditing, Warning, TEXT("Unknown value for -SANDBOXSTARTUP: %s"), *BehaviorString);
+		return false;
+	}
+}
+
+bool FStartupBehaviorHandler::TryGetNameTemplateFromCommandLine(FString& OutTemplate) const
+{
+	if (FParse::Value(FCommandLine::Get(), TEXT("SANDBOXNAMETEMPLATE="), OutTemplate))
+	{
+		OutTemplate.TrimStartAndEndInline();
+		return !OutTemplate.IsEmpty();
+	}
+	return false;
+}
+
+void FStartupBehaviorHandler::ExecuteStartupBehavior()
+{
+	const USandboxedEditingSettings* Settings = USandboxedEditingSettings::Get();
+	if (!Settings)
+	{
+		UE_LOG(LogSandboxedEditing, Warning, TEXT("Failed to get SandboxedEditingSettings"));
+		return;
+	}
+
+	// Check for command line overrides
+	ESandboxStartupBehavior Behavior = Settings->StartupBehavior;
+	if (TryGetStartupBehaviorFromCommandLine(Behavior))
+	{
+		UE_LOG(LogSandboxedEditing, Log, TEXT("Using command line override for startup behavior: %d"), static_cast<int32>(Behavior));
+	}
+
+	switch (Behavior)
+	{
+	case ESandboxStartupBehavior::DoNothing:
+		HandleDoNothing();
+		break;
+
+	case ESandboxStartupBehavior::PromptToJoin:
+		HandlePromptToJoin();
+		break;
+
+	case ESandboxStartupBehavior::CreateNewSandbox:
+		HandleCreateNewSandbox();
+		break;
+
+	case ESandboxStartupBehavior::JoinPreviousSandbox:
+		HandleJoinPreviousSandbox();
+		break;
+
+	default:
+		UE_LOG(LogSandboxedEditing, Warning, TEXT("Unknown startup behavior: %d"), static_cast<int32>(Behavior));
+		break;
+	}
+}
+
+void FStartupBehaviorHandler::HandleDoNothing()
+{
+	UE_LOG(LogSandboxedEditing, Verbose, TEXT("Startup behavior: Do Nothing"));
+	// Nothing to do - this is the default
+}
+
+void FStartupBehaviorHandler::HandlePromptToJoin()
+{
+	UE_LOG(LogSandboxedEditing, Log, TEXT("Startup behavior: Prompt to Join"));
+
+	if (!FSlateApplication::IsInitialized())
+	{
+		UE_LOG(LogSandboxedEditing, Warning, TEXT("Slate not initialized, cannot show prompt"));
+		return;
+	}
+
+	FSlateNotificationManager& NotificationManager = FSlateNotificationManager::Get();
+
+	FNotificationInfo Info(LOCTEXT("Prompt.Title", "Join Sandbox?"));
+	Info.SubText = LOCTEXT("Prompt.SubText", "Would you like to create a new sandbox or join your previous sandbox?");
+	Info.bFireAndForget = false;
+	Info.bUseLargeFont = false;
+	Info.bUseSuccessFailIcons = false;
+	Info.ButtonDetails = {
+		FNotificationButtonInfo(
+			LOCTEXT("Prompt.CreateNew", "Create New"),
+			LOCTEXT("Prompt.CreateNew.TT", "Create a new sandbox"),
+			FSimpleDelegate::CreateRaw(this, &FStartupBehaviorHandler::HandleCreateNewSandbox),
+			SNotificationItem::CS_None
+		),
+		FNotificationButtonInfo(
+			LOCTEXT("Prompt.JoinPrevious", "Join Previous"),
+			LOCTEXT("Prompt.JoinPrevious.TT", "Load your most recently used sandbox"),
+			FSimpleDelegate::CreateRaw(this, &FStartupBehaviorHandler::HandleJoinPreviousSandbox),
+			SNotificationItem::CS_None
+		),
+		FNotificationButtonInfo(
+			LOCTEXT("Prompt.Skip", "Skip"),
+			LOCTEXT("Prompt.Skip.TT", "Continue without a sandbox"),
+			FSimpleDelegate::CreateLambda([this]() {
+				if (CurrentNotification.IsValid())
+				{
+					CurrentNotification->SetCompletionState(SNotificationItem::CS_None);
+					CurrentNotification->ExpireAndFadeout();
+					CurrentNotification.Reset();
+				}
+				UE_LOG(LogSandboxedEditing, Log, TEXT("User chose to skip sandbox on startup"));
+			}),
+			SNotificationItem::CS_None
+		)
+	};
+
+	CurrentNotification = NotificationManager.AddNotification(Info);
+}
+
+void FStartupBehaviorHandler::HandleCreateNewSandbox()
+{
+	UE_LOG(LogSandboxedEditing, Log, TEXT("Startup behavior: Create New Sandbox"));
+
+	// Dismiss any active notification
+	if (CurrentNotification.IsValid())
+	{
+		CurrentNotification->ExpireAndFadeout();
+		CurrentNotification.Reset();
+	}
+
+	const USandboxedEditingSettings* Settings = USandboxedEditingSettings::Get();
+	if (!Settings)
+	{
+		ShowErrorNotification(LOCTEXT("Error.NoSettings", "Failed to read settings"));
+		return;
+	}
+
+	// Check for command line override of name template
+	FString NameTemplate;
+	if (TryGetNameTemplateFromCommandLine(NameTemplate))
+	{
+		UE_LOG(LogSandboxedEditing, Log, TEXT("Using command line override for name template: %s"), *NameTemplate);
+	}
+	else
+	{
+		NameTemplate = Settings->DefaultSandboxNameTemplate;
+	}
+
+	// Resolve naming tokens in template
+	FString ResolvedBaseName = ResolveNamingTokens(NameTemplate);
+	if (ResolvedBaseName.IsEmpty())
+	{
+		ResolvedBaseName = TEXT("Sandbox");
+	}
+
+	// Generate unique name (auto-increment if needed)
+	FString UniqueName = GenerateUniqueSandboxName(ResolvedBaseName);
+	UE_LOG(LogSandboxedEditing, Log, TEXT("Creating new sandbox with name: %s"), *UniqueName);
+
+	// Create the sandbox
+	FCreateSandboxArgs Args;
+	Args.Name = UniqueName;
+	Args.Description = FString::Printf(TEXT("Created automatically on startup at %s"),
+		*FDateTime::Now().ToString());
+
+	if (SandboxModel->CreateNewSandbox(Args))
+	{
+		ShowSuccessNotification(FText::Format(
+			LOCTEXT("Success.Created", "Created and joined sandbox '{0}'"),
+			FText::FromString(UniqueName)));
+	}
+	else
+	{
+		FText ErrorReason;
+		SandboxModel->CanCreateNewSandbox(UniqueName, &ErrorReason);
+		UE_LOG(LogSandboxedEditing, Error, TEXT("Failed to create sandbox '%s': %s"),
+			*UniqueName, *ErrorReason.ToString());
+		ShowErrorNotification(FText::Format(
+			LOCTEXT("Error.CreateFailed", "Failed to create sandbox: {0}"),
+			ErrorReason));
+	}
+}
+
+void FStartupBehaviorHandler::HandleJoinPreviousSandbox()
+{
+	UE_LOG(LogSandboxedEditing, Log, TEXT("Startup behavior: Join Previous Sandbox"));
+
+	// Dismiss any active notification
+	if (CurrentNotification.IsValid())
+	{
+		CurrentNotification->ExpireAndFadeout();
+		CurrentNotification.Reset();
+	}
+
+	const USandboxedEditingSettings* Settings = USandboxedEditingSettings::Get();
+	if (!Settings)
+	{
+		ShowErrorNotification(LOCTEXT("Error.NoSettings", "Failed to read settings"));
+		return;
+	}
+
+	// Check if we have a last used sandbox path
+	if (Settings->LastUsedSandboxPath.IsEmpty())
+	{
+		UE_LOG(LogSandboxedEditing, Warning, TEXT("No previous sandbox found, creating new sandbox instead"));
+		HandleCreateNewSandbox();
+		return;
+	}
+
+	// Check if sandbox still exists
+	const TOptional<FSandboxInfo> InfoOptional = SandboxModel->GetSandboxInfo(Settings->LastUsedSandboxPath);
+	if (!InfoOptional.IsSet())
+	{
+		UE_LOG(LogSandboxedEditing, Warning, TEXT("Previous sandbox no longer exists at path: %s, creating new sandbox instead"),
+			*Settings->LastUsedSandboxPath);
+
+		// Clear stale path
+		USandboxedEditingSettings* MutableSettings = USandboxedEditingSettings::Get();
+		MutableSettings->LastUsedSandboxPath.Empty();
+		MutableSettings->SaveConfig();
+
+		// Fallback to creating a new sandbox
+		HandleCreateNewSandbox();
+		return;
+	}
+
+	const FSandboxInfo& Info = InfoOptional.GetValue();
+	UE_LOG(LogSandboxedEditing, Log, TEXT("Loading previous sandbox: %s (path: %s)"),
+		*Info.Name, *Info.SandboxRoot);
+
+	// Load the sandbox
+	if (SandboxModel->LoadSandbox(Info.SandboxRoot))
+	{
+		ShowSuccessNotification(FText::Format(
+			LOCTEXT("Success.Loaded", "Loaded previous sandbox '{0}'"),
+			FText::FromString(Info.Name)));
+	}
+	else
+	{
+		UE_LOG(LogSandboxedEditing, Error, TEXT("Failed to load previous sandbox: %s, creating new sandbox instead"), *Info.Name);
+
+		// Fallback to creating a new sandbox
+		HandleCreateNewSandbox();
+	}
+}
+
+FString FStartupBehaviorHandler::ResolveNamingTokens(const FString& Template)
+{
+	if (Template.IsEmpty())
+	{
+		UE_LOG(LogSandboxedEditing, Verbose, TEXT("Empty template, using fallback"));
+		return FallbackToTimestamp();
+	}
+
+	// Check if GEngine is available
+	if (GEngine == nullptr)
+	{
+		UE_LOG(LogSandboxedEditing, Warning, TEXT("GEngine not available, using fallback naming"));
+		return FallbackToTimestamp();
+	}
+
+	// Get the NamingTokens subsystem
+	UNamingTokensEngineSubsystem* Subsystem =
+		GEngine->GetEngineSubsystem<UNamingTokensEngineSubsystem>();
+
+	if (Subsystem == nullptr)
+	{
+		UE_LOG(LogSandboxedEditing, Warning, TEXT("NamingTokens subsystem not available, using fallback naming"));
+		return FallbackToTimestamp();
+	}
+
+	// Evaluate the template
+	FNamingTokenResultData Result = Subsystem->EvaluateTokenString(Template);
+	FString EvaluatedName = Result.EvaluatedText.ToString();
+
+	if (EvaluatedName.IsEmpty())
+	{
+		UE_LOG(LogSandboxedEditing, Warning, TEXT("Token evaluation resulted in empty string, using fallback"));
+		return FallbackToTimestamp();
+	}
+
+	UE_LOG(LogSandboxedEditing, Verbose, TEXT("Resolved naming tokens: '%s' -> '%s'"),
+		*Template, *EvaluatedName);
+
+	return EvaluatedName;
+}
+
+FString FStartupBehaviorHandler::FallbackToTimestamp()
+{
+	FDateTime Now = FDateTime::Now();
+	FString Fallback = FString::Printf(TEXT("Sandbox_%04d_%02d_%02d_%02d%02d"),
+		Now.GetYear(), Now.GetMonth(), Now.GetDay(),
+		Now.GetHour(), Now.GetMinute());
+
+	UE_LOG(LogSandboxedEditing, Verbose, TEXT("Using fallback timestamp name: %s"), *Fallback);
+	return Fallback;
+}
+
+FString FStartupBehaviorHandler::GenerateUniqueSandboxName(const FString& BaseName)
+{
+	FText ErrorReason;
+
+	// Try the base name first
+	if (SandboxModel->CanCreateNewSandbox(BaseName, &ErrorReason))
+	{
+		return BaseName;
+	}
+
+	UE_LOG(LogSandboxedEditing, Log, TEXT("Base name '%s' already exists, trying auto-increment"), *BaseName);
+
+	// Try with suffixes: BaseName_1, BaseName_2, etc.
+	for (int32 Suffix = 1; Suffix < 1000; ++Suffix)
+	{
+		FString Candidate = FString::Printf(TEXT("%s_%d"), *BaseName, Suffix);
+		if (SandboxModel->CanCreateNewSandbox(Candidate, &ErrorReason))
+		{
+			UE_LOG(LogSandboxedEditing, Log, TEXT("Found unique name with suffix: %s"), *Candidate);
+			return Candidate;
+		}
+	}
+
+	// Fallback to GUID for extreme edge case
+	FString Unique = FString::Printf(TEXT("%s_%s"), *BaseName,
+		*FGuid::NewGuid().ToString(EGuidFormats::Short));
+	UE_LOG(LogSandboxedEditing, Warning,
+		TEXT("Generated GUID-based name after 1000 iterations: %s"), *Unique);
+	return Unique;
+}
+
+void FStartupBehaviorHandler::ShowSuccessNotification(const FText& Message)
+{
+	if (!FSlateApplication::IsInitialized())
+	{
+		return;
+	}
+
+	FSlateNotificationManager& NotificationManager = FSlateNotificationManager::Get();
+	FNotificationInfo Info(Message);
+	Info.bFireAndForget = true;
+	Info.ExpireDuration = 3.0f;
+	Info.bUseLargeFont = false;
+
+	TSharedPtr<SNotificationItem> Item = NotificationManager.AddNotification(Info);
+	if (Item.IsValid())
+	{
+		Item->SetCompletionState(SNotificationItem::CS_Success);
+	}
+}
+
+void FStartupBehaviorHandler::ShowErrorNotification(const FText& Message)
+{
+	if (!FSlateApplication::IsInitialized())
+	{
+		return;
+	}
+
+	FSlateNotificationManager& NotificationManager = FSlateNotificationManager::Get();
+	FNotificationInfo Info(Message);
+	Info.bFireAndForget = true;
+	Info.ExpireDuration = 5.0f;
+	Info.bUseLargeFont = false;
+
+	TSharedPtr<SNotificationItem> Item = NotificationManager.AddNotification(Info);
+	if (Item.IsValid())
+	{
+		Item->SetCompletionState(SNotificationItem::CS_Fail);
+	}
+}
+
+void FStartupBehaviorHandler::OnSandboxLoaded()
+{
+	if (!SandboxModel->HasActiveSandbox())
+	{
+		return;
+	}
+
+	FString CurrentPath = SandboxModel->GetActiveSandboxPath();
+	USandboxedEditingSettings* Settings = USandboxedEditingSettings::Get();
+
+	if (!Settings)
+	{
+		return;
+	}
+
+	// Update last used path if it changed
+	if (Settings->LastUsedSandboxPath != CurrentPath)
+	{
+		Settings->LastUsedSandboxPath = CurrentPath;
+		Settings->SaveConfig();
+		UE_LOG(LogSandboxedEditing, Verbose, TEXT("Updated last used sandbox path: %s"), *CurrentPath);
+	}
+}
+
+} // namespace UE::SandboxedEditing
+
+#undef LOCTEXT_NAMESPACE
diff --git a/Engine/Plugins/Developer/Sandbox/SandboxedEditing/Source/SandboxedEditing/Private/Framework/StartupBehavior/StartupBehaviorHandler.h b/Engine/Plugins/Developer/Sandbox/SandboxedEditing/Source/SandboxedEditing/Private/Framework/StartupBehavior/StartupBehaviorHandler.h
new file mode 100644
--- /dev/null
+++ b/Engine/Plugins/Developer/Sandbox/SandboxedEditing/Source/SandboxedEditing/Private/Framework/StartupBehavior/StartupBehaviorHandler.h
@@ -0,0 +1,119 @@
+// Copyright Epic Games, Inc. All Rights Reserved.
+
+#pragma once
+
+#include "Templates/SharedPointer.h"
+#include "SandboxedEditingSettings.h"
+#include "Delegates/IDelegateInstance.h"
+
+class SNotificationItem;
+
+namespace UE::SandboxedEditing
+{
+class FSandboxSystemModel;
+
+/**
+ * Handles sandbox startup behavior based on user settings.
+ * Automatically creates, loads, or prompts for sandboxes when the editor starts.
+ * Also tracks the last used sandbox for the "Join Previous" workflow.
+ */
+class FStartupBehaviorHandler : public FNoncopyable
+{
+public:
+	/**
+	 * Constructor - registers for editor initialization and sandbox load events
+	 * @param InModel The sandbox system model to use for operations
+	 */
+	explicit FStartupBehaviorHandler(const TSharedRef<FSandboxSystemModel>& InModel);
+
+	/** Destructor - unregisters all delegates */
+	~FStartupBehaviorHandler();
+
+private:
+	// Editor initialization
+	/** Called when the editor has fully initialized and is ready */
+	void OnEditorInitialized(double Duration);
+
+	/** Reads settings and executes the appropriate startup behavior */
+	void ExecuteStartupBehavior();
+
+	// Behavior handlers
+	/** Does nothing - default behavior */
+	void HandleDoNothing();
+
+	/** Shows a prompt with options to create new, join previous, or skip */
+	void HandlePromptToJoin();
+
+	/** Automatically creates a new sandbox using the template name */
+	void HandleCreateNewSandbox();
+
+	/** Automatically loads the most recently used sandbox */
+	void HandleJoinPreviousSandbox();
+
+	// Command line parsing
+	/**
+	 * Parses command line override for startup behavior
+	 * @param OutBehavior If command line override is present, this will be set
+	 * @return True if command line override was found, false otherwise
+	 */
+	bool TryGetStartupBehaviorFromCommandLine(ESandboxStartupBehavior& OutBehavior) const;
+
+	/**
+	 * Parses command line override for sandbox name template
+	 * @param OutTemplate If command line override is present, this will be set
+	 * @return True if command line override was found, false otherwise
+	 */
+	bool TryGetNameTemplateFromCommandLine(FString& OutTemplate) const;
+
+	// Utilities
+	/**
+	 * Resolves naming tokens in a template string using the NamingTokens plugin
+	 * @param Template The template string with tokens like {yyyy}, {mm}, {dd}, {user}
+	 * @return The resolved string with tokens replaced, or a fallback if resolution fails
+	 */
+	FString ResolveNamingTokens(const FString& Template);
+
+	/**
+	 * Generates a fallback sandbox name using timestamp when NamingTokens fails
+	 * @return A timestamp-based sandbox name like "Sandbox_2026_03_27_1435"
+	 */
+	FString FallbackToTimestamp();
+
+	/**
+	 * Generates a unique sandbox name by auto-incrementing if the base name exists
+	 * @param BaseName The desired base name
+	 * @return A unique name (BaseName, BaseName_1, BaseName_2, etc.)
+	 */
+	FString GenerateUniqueSandboxName(const FString& BaseName);
+
+	/**
+	 * Shows a success notification
+	 * @param Message The message to display
+	 */
+	void ShowSuccessNotification(const FText& Message);
+
+	/**
+	 * Shows an error notification
+	 * @param Message The error message to display
+	 */
+	void ShowErrorNotification(const FText& Message);
+
+	// Last used tracking
+	/**
+	 * Called whenever a sandbox is loaded (manual or automatic)
+	 * Updates LastUsedSandboxPath in settings
+	 */
+	void OnSandboxLoaded();
+
+	// Members
+	/** The sandbox system model for all sandbox operations */
+	TSharedRef<FSandboxSystemModel> SandboxModel;
+
+	/** Current notification being displayed (for prompt workflow) */
+	TSharedPtr<SNotificationItem> CurrentNotification;
+
+	/** Delegate handle for sandbox load tracking */
+	FDelegateHandle LoadSandboxHandle;
+};
+
+} // namespace UE::SandboxedEditing
diff --git a/Engine/Plugins/Developer/Sandbox/SandboxedEditing/Source/SandboxedEditing/Private/SandboxedEditingSettings.h b/Engine/Plugins/Developer/Sandbox/SandboxedEditing/Source/SandboxedEditing/Private/SandboxedEditingSettings.h
--- a/Engine/Plugins/Developer/Sandbox/SandboxedEditing/Source/SandboxedEditing/Private/SandboxedEditingSettings.h
+++ b/Engine/Plugins/Developer/Sandbox/SandboxedEditing/Source/SandboxedEditing/Private/SandboxedEditingSettings.h
@@ -5,17 +5,54 @@
 #include "UObject/UObjectGlobals.h"
 #include "SandboxedEditingSettings.generated.h"
 
+/**
+ * Determines how the editor should behave on startup with regards to sandboxes.
+ *
+ * Command Line Overrides:
+ *   -SANDBOXSTARTUP=DoNothing|PromptToJoin|CreateNewSandbox|JoinPreviousSandbox
+ *   -SANDBOXNAMETEMPLATE="Sandbox_{yyyy}_{mm}_{dd}"
+ *
+ * Command line overrides take precedence over these settings and are not saved to config.
+ */
+UENUM()
+enum class ESandboxStartupBehavior : uint8
+{
+	/** Do not automatically create or join a sandbox on startup (default) */
+	DoNothing UMETA(DisplayName = "Do Nothing"),
+
+	/** Show a prompt asking the user if they want to create or join a sandbox */
+	PromptToJoin UMETA(DisplayName = "Prompt to Join"),
+
+	/** Automatically create a new sandbox using the default template name */
+	CreateNewSandbox UMETA(DisplayName = "Create New Sandbox"),
+
+	/** Automatically join the most recently used sandbox */
+	JoinPreviousSandbox UMETA(DisplayName = "Join Previous Sandbox")
+};
+
 UCLASS(MinimalAPI, config=Engine)
 class USandboxedEditingSettings : public UObject
 {
 	GENERATED_BODY()
 public:
-	
+
 	static USandboxedEditingSettings* Get() { return GetMutableDefault<USandboxedEditingSettings>(); }
-	
+
 	/** Should we always ask the user to persist sandbox contents when leaving a sandbox? */
 	UPROPERTY(config, EditAnywhere, Category = "Actions", meta = (DisplayName="Ask to persist when leaving the sandbox"))
 	bool bAskToPersistWhenLeavingSandbox = true;
+
+	/** Controls what happens on editor startup with regards to sandboxes. Can be overridden with -SANDBOXSTARTUP=DoNothing|PromptToJoin|CreateNewSandbox|JoinPreviousSandbox */
+	UPROPERTY(config, EditAnywhere, Category = "Startup", meta = (DisplayName = "Startup Behavior", ToolTip = "Controls what happens on editor startup. Can be overridden with -SANDBOXSTARTUP command line argument."))
+	ESandboxStartupBehavior StartupBehavior = ESandboxStartupBehavior::DoNothing;
+
+	/** Template for sandbox names when creating new sandboxes on startup. Supports {yyyy}, {mm}, {dd}, {24h}, {min}, {sec}, {user}, etc. Can be overridden with -SANDBOXNAMETEMPLATE */
+	UPROPERTY(config, EditAnywhere, Category = "Startup", meta = (DisplayName = "Default Sandbox Name Template", EditCondition = "StartupBehavior != ESandboxStartupBehavior::DoNothing", EditConditionHides, ToolTip = "Template for sandbox names. Supports {yyyy}, {mm}, {dd}, {24h}, {min}, {sec}, {user}, etc. Example: Sandbox_{yyyy}_{mm}_{dd}_{user}. Can be overridden with -SANDBOXNAMETEMPLATE command line argument."))
+	FString DefaultSandboxNameTemplate = TEXT("Sandbox_{yyyy}_{mm}_{dd}");
+
+	/** Internal tracking of the last used sandbox path (auto-updated, not user-editable) */
+	UPROPERTY(config)
+	FString LastUsedSandboxPath;
 };
 
 
diff --git a/Engine/Plugins/Developer/Sandbox/SandboxedEditing/Source/SandboxedEditing/SandboxedEditing.Build.cs b/Engine/Plugins/Developer/Sandbox/SandboxedEditing/Source/SandboxedEditing/SandboxedEditing.Build.cs
--- a/Engine/Plugins/Developer/Sandbox/SandboxedEditing/Source/SandboxedEditing/SandboxedEditing.Build.cs
+++ b/Engine/Plugins/Developer/Sandbox/SandboxedEditing/Source/SandboxedEditing/SandboxedEditing.Build.cs
@@ -19,16 +19,19 @@ namespace UnrealBuildTool.Rules
 				{
 					"AssetTools",
 					"DesktopPlatform",
+					"Engine",
 					"FileSandboxCore",
 					"FileSandboxUI",
-					"FileUtilities", 
+					"FileUtilities",
 					"InputCore",
-					"Slate", 
-					"SlateCore", 
+					"NamingTokens",
+					"Slate",
+					"SlateCore",
 					"SourceControl",
-					"ToolMenus", 
-					"ToolWidgets", 
-					"WorkspaceMenuStructure", 
+					"ToolMenus",
+					"ToolWidgets",
+					"UnrealEd",
+					"WorkspaceMenuStructure",
 				}
 			);
 		}
