Implements opt-in feature to delete sandboxes with no file changes when
the editor shuts down. Reduces clutter from unused experimental sandboxes.

New setting: "Delete Empty Sandboxes on Exit" (default: false)
Runs on FCoreDelegates::OnEnginePreExit shutdown hook
Leaves active sandbox first if empty, then deletes it
Checks all known sandboxes for file changes
Preserves sandboxes with created/modified/deleted files
Logs all deletions

#rnx

#Synced-CL 52173051

diff --git a/Engine/Plugins/Developer/Sandbox/SandboxedEditing/Source/SandboxedEditing/Private/Framework/Cleanup/ShutdownCleanupHandler.cpp b/Engine/Plugins/Developer/Sandbox/SandboxedEditing/Source/SandboxedEditing/Private/Framework/Cleanup/ShutdownCleanupHandler.cpp
new file mode 100644
--- /dev/null
+++ b/Engine/Plugins/Developer/Sandbox/SandboxedEditing/Source/SandboxedEditing/Private/Framework/Cleanup/ShutdownCleanupHandler.cpp
@@ -0,0 +1,88 @@
+// Copyright Epic Games, Inc. All Rights Reserved.
+
+#include "ShutdownCleanupHandler.h"
+
+#include "Misc/CoreDelegates.h"
+#include "Framework/Models/SandboxInfo.h"
+#include "Framework/Models/SandboxSystemModel.h"
+#include "Logging/LogMacros.h"
+#include "SandboxedEditingSettings.h"
+#include "Types/GatheredFileChanges.h"
+
+DEFINE_LOG_CATEGORY_STATIC(LogSandboxCleanup, Log, All);
+
+namespace UE::SandboxedEditing
+{
+
+FShutdownCleanupHandler::FShutdownCleanupHandler(const TSharedRef<FSandboxSystemModel>& InModel)
+	: SandboxModel(InModel)
+{
+	// Register for earliest shutdown hook
+	ShutdownHandle = FCoreDelegates::OnEnginePreExit.AddRaw(this, &FShutdownCleanupHandler::OnEnginePreExit);
+}
+
+FShutdownCleanupHandler::~FShutdownCleanupHandler()
+{
+	if (ShutdownHandle.IsValid())
+	{
+		FCoreDelegates::OnEnginePreExit.Remove(ShutdownHandle);
+	}
+}
+
+void FShutdownCleanupHandler::OnEnginePreExit()
+{
+	const USandboxedEditingSettings* Settings = USandboxedEditingSettings::Get();
+
+	// Early exit if feature disabled or settings unavailable
+	if (!Settings || !Settings->bDeleteEmptySandboxesOnExit)
+	{
+		return;
+	}
+
+	UE_LOG(LogSandboxCleanup, Log, TEXT("Starting empty sandbox cleanup on shutdown"));
+
+	// Leave active sandbox first so we can check/delete it
+	if (SandboxModel->HasActiveSandbox())
+	{
+		FString ActiveSandboxPath = SandboxModel->GetActiveSandboxPath();
+		UE_LOG(LogSandboxCleanup, Verbose, TEXT("Leaving active sandbox to check for cleanup: %s"), *ActiveSandboxPath);
+		SandboxModel->LeaveSandbox();
+	}
+
+	// Get all known sandboxes
+	TArray<FSandboxInfo> AllSandboxes = SandboxModel->GetKnownSandboxes();
+	int32 DeletedCount = 0;
+
+	for (const FSandboxInfo& Info : AllSandboxes)
+	{
+		// Check if sandbox is empty
+		if (!HasFileChanges(Info.SandboxRoot))
+		{
+			UE_LOG(LogSandboxCleanup, Log, TEXT("Deleting empty sandbox: %s (Path: %s)"), *Info.Name, *Info.SandboxRoot);
+
+			// Delete sandbox
+			SandboxModel->DeleteSandbox(Info.SandboxRoot);
+			DeletedCount++;
+		}
+	}
+
+	if (DeletedCount > 0)
+	{
+		UE_LOG(LogSandboxCleanup, Log, TEXT("Cleanup complete: Deleted %d empty sandbox(es)"), DeletedCount);
+	}
+	else
+	{
+		UE_LOG(LogSandboxCleanup, Log, TEXT("Cleanup complete: No empty sandboxes found"));
+	}
+}
+
+bool FShutdownCleanupHandler::HasFileChanges(const FString& SandboxRoot)
+{
+	// Gather file changes from the sandbox
+	FileSandboxCore::FGatheredFileChanges Changes = SandboxModel->GatherFileChanges(SandboxRoot);
+
+	// HasChanges() returns true if there are any created/modified/deleted files
+	return Changes.HasChanges();
+}
+
+} // namespace UE::SandboxedEditing
diff --git a/Engine/Plugins/Developer/Sandbox/SandboxedEditing/Source/SandboxedEditing/Private/Framework/Cleanup/ShutdownCleanupHandler.h b/Engine/Plugins/Developer/Sandbox/SandboxedEditing/Source/SandboxedEditing/Private/Framework/Cleanup/ShutdownCleanupHandler.h
new file mode 100644
--- /dev/null
+++ b/Engine/Plugins/Developer/Sandbox/SandboxedEditing/Source/SandboxedEditing/Private/Framework/Cleanup/ShutdownCleanupHandler.h
@@ -0,0 +1,46 @@
+// Copyright Epic Games, Inc. All Rights Reserved.
+
+#pragma once
+
+#include "Delegates/IDelegateInstance.h"
+#include "Templates/SharedPointer.h"
+
+namespace UE::SandboxedEditing
+{
+class FSandboxSystemModel;
+
+/**
+ * Handles cleanup of empty sandboxes on editor shutdown.
+ * Checks all known sandboxes and deletes those with no file changes.
+ */
+class FShutdownCleanupHandler : public FNoncopyable
+{
+public:
+	/**
+	 * Constructor - registers for engine shutdown event
+	 * @param InModel The sandbox system model to use for operations
+	 */
+	explicit FShutdownCleanupHandler(const TSharedRef<FSandboxSystemModel>& InModel);
+
+	/** Destructor - unregisters delegates */
+	~FShutdownCleanupHandler();
+
+private:
+	/** Called when engine is about to shut down */
+	void OnEnginePreExit();
+
+	/**
+	 * Checks if a sandbox has any file changes
+	 * @param SandboxRoot The root directory of the sandbox
+	 * @return True if the sandbox has created/modified/deleted files
+	 */
+	bool HasFileChanges(const FString& SandboxRoot);
+
+	/** The sandbox system model for all sandbox operations */
+	TSharedRef<FSandboxSystemModel> SandboxModel;
+
+	/** Delegate handle for cleanup */
+	FDelegateHandle ShutdownHandle;
+};
+
+} // namespace UE::SandboxedEditing
diff --git a/Engine/Plugins/Developer/Sandbox/SandboxedEditing/Source/SandboxedEditing/Private/Framework/SandboxedEditingApp.cpp b/Engine/Plugins/Developer/Sandbox/SandboxedEditing/Source/SandboxedEditing/Private/Framework/SandboxedEditingApp.cpp
--- a/Engine/Plugins/Developer/Sandbox/SandboxedEditing/Source/SandboxedEditing/Private/Framework/SandboxedEditingApp.cpp
+++ b/Engine/Plugins/Developer/Sandbox/SandboxedEditing/Source/SandboxedEditing/Private/Framework/SandboxedEditingApp.cpp
@@ -15,6 +15,7 @@ FSandboxedEditingApp::FSandboxedEditingApp()
 	, StatusBarFeature(SandboxModel)
 	, BrowserFeature(SandboxModel, SandboxIntegration->GetExternalSandboxViewModel())
 	, StartupBehaviorHandler(SandboxModel)
+	, ShutdownCleanupHandler(SandboxModel)
 {
 	SandboxIntegration->GetEntryPoint()->OnRequestSummonUI().AddLambda([this] { BrowserFeature.SummonUI(); });
 }
diff --git a/Engine/Plugins/Developer/Sandbox/SandboxedEditing/Source/SandboxedEditing/Private/Framework/SandboxedEditingApp.h b/Engine/Plugins/Developer/Sandbox/SandboxedEditing/Source/SandboxedEditing/Private/Framework/SandboxedEditingApp.h
--- a/Engine/Plugins/Developer/Sandbox/SandboxedEditing/Source/SandboxedEditing/Private/Framework/SandboxedEditingApp.h
+++ b/Engine/Plugins/Developer/Sandbox/SandboxedEditing/Source/SandboxedEditing/Private/Framework/SandboxedEditingApp.h
@@ -4,6 +4,7 @@
 
 #include "Features/Browser/BrowserFeature.h"
 #include "Features/StatusBar/StatusBarFeature.h"
+#include "Framework/Cleanup/ShutdownCleanupHandler.h"
 #include "Framework/StartupBehavior/StartupBehaviorHandler.h"
 #include "Templates/SharedPointer.h"
 #include "Templates/UnrealTemplate.h"
@@ -40,6 +41,8 @@ private:
 	const FBrowserFeature BrowserFeature;
 	/** Handles startup behavior (creating/loading sandboxes on editor startup). */
 	const FStartupBehaviorHandler StartupBehaviorHandler;
+	/** Handles cleanup of empty sandboxes on editor shutdown. */
+	const FShutdownCleanupHandler ShutdownCleanupHandler;
 };
 }
 
diff --git a/Engine/Plugins/Developer/Sandbox/SandboxedEditing/Source/SandboxedEditing/Private/SandboxedEditingSettings.h b/Engine/Plugins/Developer/Sandbox/SandboxedEditing/Source/SandboxedEditing/Private/SandboxedEditingSettings.h
--- a/Engine/Plugins/Developer/Sandbox/SandboxedEditing/Source/SandboxedEditing/Private/SandboxedEditingSettings.h
+++ b/Engine/Plugins/Developer/Sandbox/SandboxedEditing/Source/SandboxedEditing/Private/SandboxedEditingSettings.h
@@ -53,6 +53,10 @@ public:
 	/** Internal tracking of the last used sandbox path (auto-updated, not user-editable) */
 	UPROPERTY(config)
 	FString LastUsedSandboxPath;
+
+	/** Automatically delete sandboxes with no file changes when the editor shuts down */
+	UPROPERTY(config, EditAnywhere, Category = "Cleanup", meta = (DisplayName = "Delete Empty Sandboxes on Exit", ToolTip = "Automatically delete sandboxes with no file changes when the editor shuts down"))
+	bool bDeleteEmptySandboxesOnExit = false;
 };
 
 
