Add custom sandbox storage directory setting

Implements configurable storage location for sandboxes via Project Settings,
allowing users to specify a custom directory instead of the default
{Project}/Intermediate/Sandboxes/ location.

- New setting: CustomSandboxStorageDirectory (FDirectoryPath) in Storage category
- Empty by default (uses {Project}/Intermediate/Sandboxes/)
- Command line override (-DEFAULTSANDBOXDIRECTORY) takes precedence
- Requires editor restart to take effect
- Tooltip documents default location and restart requirement

#rnx
#Synced-CL 52173051

diff --git a/Engine/Plugins/Developer/Sandbox/SandboxedEditing/Source/SandboxedEditing/Private/Framework/Models/SandboxSystemModel.cpp b/Engine/Plugins/Developer/Sandbox/SandboxedEditing/Source/SandboxedEditing/Private/Framework/Models/SandboxSystemModel.cpp
--- a/Engine/Plugins/Developer/Sandbox/SandboxedEditing/Source/SandboxedEditing/Private/Framework/Models/SandboxSystemModel.cpp
+++ b/Engine/Plugins/Developer/Sandbox/SandboxedEditing/Source/SandboxedEditing/Private/Framework/Models/SandboxSystemModel.cpp
@@ -13,6 +13,7 @@
 #include "ISandboxRepository.h"
 #include "Interface/ISandboxLock.h"
 #include "Misc/Paths.h"
+#include "SandboxedEditingSettings.h"
 #include "SandboxTagging.h"
 #include "Templates/Function.h"
 #include "Types/EBreakBehavior.h"
@@ -23,6 +24,7 @@
 #include "Types/Sandbox/RevertResult.h"
 #include "Types/SandboxMetaInfo.h"
 #include "Utils/PackageSandboxUtils.h"
+#include "Utils/SandboxDirectoryUtils.h"
 
 #define LOCTEXT_NAMESPACE "FSandboxSystemModel"
 
@@ -81,8 +83,16 @@ bool FSandboxSystemModel::CreateNewSandbox(const FCreateSandboxArgs& InArgs)
 	}
 
 	FNewSandboxArgs Args(InArgs.Name, InArgs.Description);
+	Args.SandboxBasePath = SandboxModel::GetBaseSandboxDirectory();
 	MarkOwnedBySandboxedEditing(Args.MetaData);
-	return SandboxManager.CreateNewSandbox(Args).HasValue();
+	const bool bSuccess = SandboxManager.CreateNewSandbox(Args).HasValue();
+
+	if (bSuccess)
+	{
+		NotifyIfBypassingRepository();
+	}
+
+	return bSuccess;
 }
 
 bool FSandboxSystemModel::CanCreateNewSandbox() const
@@ -94,12 +104,13 @@ bool FSandboxSystemModel::CanCreateNewSandbox() const
 	
 bool FSandboxSystemModel::CanCreateNewSandbox(const FString& InName, FText* Reason)
 {
-	FileSandboxCore::EDirectorySuitability Suitability = FileSandboxCore::DetermineDirectorySuitability(InName);
+	const FString BaseDir = SandboxModel::GetBaseSandboxDirectory();
+	FileSandboxCore::EDirectorySuitability Suitability = FileSandboxCore::DetermineDirectorySuitability(InName, BaseDir);
 
 	const bool bIsSuitable = Suitability == FileSandboxCore::EDirectorySuitability::Suitable;
 	if (Reason && !bIsSuitable)
 	{
-		*Reason = FileSandboxCore::FormatDirectorySuitabilityAsReason(Suitability, InName);
+		*Reason = FileSandboxCore::FormatDirectorySuitabilityAsReason(Suitability, InName, BaseDir);
 	}
 	
 	if (HasActiveSandbox() && !IsAllowedToLeaveSandbox(Reason))
@@ -216,6 +227,7 @@ void FSandboxSystemModel::DeleteSandbox(const FString& InSandboxRoot) const
 		IFileSandboxCoreModule& Module = IFileSandboxCoreModule::Get();
 		ISandboxManager& SandboxManager = Module.GetSandboxManager();
 		SandboxManager.DeleteSandbox(FDeleteSandboxByDirectoryArgs{ InSandboxRoot });
+		NotifyIfBypassingRepository();
 	}
 }
 
@@ -351,18 +363,37 @@ bool FSandboxSystemModel::HasActiveSandbox() const
 void FSandboxSystemModel::ForEachSandbox(TFunctionRef<void(const FSandboxInfo& InRoot)> InProcess) const
 {
 	using namespace FileSandboxCore;
-	IFileSandboxCoreModule& Module = IFileSandboxCoreModule::Get();
-	
-	ISandboxRepository& Repository = Module.GetDefaultSandboxRepository();
-	Repository.ForEachSandbox([&InProcess](const FString& InRootPath, const FSandboxMetaInfo& MetaInfo)
+
+	if (IsUsingDefaultRepository())
 	{
-		const FFileSandboxCore_SandboxMetaData& MetaData = MetaInfo.UserMetaData;
-		if (CanBeShownBySandboxedEditing(MetaData))
+		// Use repository for default location (benefits from caching and directory watching)
+		IFileSandboxCoreModule& Module = IFileSandboxCoreModule::Get();
+		ISandboxRepository& Repository = Module.GetDefaultSandboxRepository();
+		Repository.ForEachSandbox([&InProcess](const FString& InRootPath, const FSandboxMetaInfo& MetaInfo)
 		{
-			InProcess(FSandboxInfo(MetaData.Name, MetaData.Description, InRootPath, GetLastModified(InRootPath), MetaInfo.VersionInfo.EngineVersion));
-		}
-		return FileSandboxCore::EBreakBehavior::Continue;
-	});
+			const FFileSandboxCore_SandboxMetaData& MetaData = MetaInfo.UserMetaData;
+			if (CanBeShownBySandboxedEditing(MetaData))
+			{
+				InProcess(FSandboxInfo(MetaData.Name, MetaData.Description, InRootPath, GetLastModified(InRootPath), MetaInfo.VersionInfo.EngineVersion));
+			}
+			return FileSandboxCore::EBreakBehavior::Continue;
+		});
+	}
+	else
+	{
+		// Bypass repository and enumerate custom directory directly
+		const FString CustomBaseDir = SandboxModel::GetBaseSandboxDirectory();
+		FileSandboxCore::ForEachSandbox([&InProcess](const FString& InRootPath)
+		{
+			const TOptional<FFileSandboxCore_SandboxMetaData> MetaData = LoadMetaData(InRootPath);
+			if (MetaData && CanBeShownBySandboxedEditing(*MetaData))
+			{
+				const FFileSandboxCore_VersionInfo VersionInfo = LoadVersionInfo(InRootPath);
+				InProcess(FSandboxInfo(MetaData->Name, MetaData->Description, InRootPath, GetLastModified(InRootPath), VersionInfo.EngineVersion));
+			}
+			return FileSandboxCore::EBreakBehavior::Continue;
+		}, CustomBaseDir);
+	}
 }
 
 TArray<FSandboxInfo> FSandboxSystemModel::GetKnownSandboxes() const
@@ -375,42 +406,82 @@ TArray<FSandboxInfo> FSandboxSystemModel::GetKnownSandboxes() const
 TOptional<FSandboxInfo> FSandboxSystemModel::GetSandboxInfo(const FString& InSandboxRoot) const
 {
 	using namespace FileSandboxCore;
-	IFileSandboxCoreModule& Module = IFileSandboxCoreModule::Get();
-	ISandboxRepository& Repository = Module.GetDefaultSandboxRepository();
-	
+
 	TOptional<FSandboxInfo> Result;
-	Repository.ReadMetaData(InSandboxRoot, [&InSandboxRoot, &Result](const FSandboxMetaInfo& MetaInfo)
+
+	if (IsUsingDefaultRepository())
 	{
-		const FFileSandboxCore_SandboxMetaData& MetaData = MetaInfo.UserMetaData;
-		if (CanBeShownBySandboxedEditing(MetaData))
+		// Use repository for default location
+		IFileSandboxCoreModule& Module = IFileSandboxCoreModule::Get();
+		ISandboxRepository& Repository = Module.GetDefaultSandboxRepository();
+		Repository.ReadMetaData(InSandboxRoot, [&InSandboxRoot, &Result](const FSandboxMetaInfo& MetaInfo)
 		{
-			Result.Emplace(MetaData.Name, MetaData.Description, InSandboxRoot, MetaInfo.LastModified, MetaInfo.VersionInfo.EngineVersion);
+			const FFileSandboxCore_SandboxMetaData& MetaData = MetaInfo.UserMetaData;
+			if (CanBeShownBySandboxedEditing(MetaData))
+			{
+				Result.Emplace(MetaData.Name, MetaData.Description, InSandboxRoot, MetaInfo.LastModified, MetaInfo.VersionInfo.EngineVersion);
+			}
+		});
+	}
+	else
+	{
+		// Bypass repository and load directly from custom directory
+		const TOptional<FFileSandboxCore_SandboxMetaData> MetaData = LoadMetaData(InSandboxRoot);
+		if (MetaData && CanBeShownBySandboxedEditing(*MetaData))
+		{
+			const FFileSandboxCore_VersionInfo VersionInfo = LoadVersionInfo(InSandboxRoot);
+			Result.Emplace(MetaData->Name, MetaData->Description, InSandboxRoot, GetLastModified(InSandboxRoot), VersionInfo.EngineVersion);
 		}
-	});
+	}
+
 	return Result;
 }
 
 void FSandboxSystemModel::SetDescription(const FString& InSandboxRoot, const FString& InDescription) const
 {
 	using namespace FileSandboxCore;
-	IFileSandboxCoreModule& Module = IFileSandboxCoreModule::Get();
-	ISandboxRepository& Repository = Module.GetDefaultSandboxRepository();
-	Repository.ReadMetaData(InSandboxRoot, [&InSandboxRoot, &InDescription](const FSandboxMetaInfo& MetaInfo)
+
+	if (IsUsingDefaultRepository())
+	{
+		// Use repository for default location (repository will fire events)
+		IFileSandboxCoreModule& Module = IFileSandboxCoreModule::Get();
+		ISandboxRepository& Repository = Module.GetDefaultSandboxRepository();
+		Repository.ReadMetaData(InSandboxRoot, [&InSandboxRoot, &InDescription](const FSandboxMetaInfo& MetaInfo)
+		{
+			if (CanBeShownBySandboxedEditing(MetaInfo.UserMetaData))
+			{
+				FFileSandboxCore_SandboxMetaData MetaData = MetaInfo.UserMetaData;
+				MetaData.Description = InDescription;
+				SaveMetaData(MetaData, InSandboxRoot);
+			}
+		});
+	}
+	else
 	{
-		if (CanBeShownBySandboxedEditing(MetaInfo.UserMetaData))
+		// Bypass repository and load/save directly from custom directory
+		const TOptional<FFileSandboxCore_SandboxMetaData> MetaData = LoadMetaData(InSandboxRoot);
+		if (MetaData && CanBeShownBySandboxedEditing(*MetaData))
 		{
-			FFileSandboxCore_SandboxMetaData MetaData = MetaInfo.UserMetaData;
-			MetaData.Description = InDescription;
-			SaveMetaData(MetaData, InSandboxRoot);
+			FFileSandboxCore_SandboxMetaData UpdatedMetaData = *MetaData;
+			UpdatedMetaData.Description = InDescription;
+			SaveMetaData(UpdatedMetaData, InSandboxRoot);
+			NotifyIfBypassingRepository();
 		}
-	});
+	}
 }
 
 bool FSandboxSystemModel::RenameSandbox(const FString& InSandboxRoot, const FString& InNewName) const
 {
 	const FString OldSandboxName = FPaths::GetBaseFilename(InSandboxRoot);
 	const FString BaseDirectory = FPaths::GetPath(InSandboxRoot);
-	return FileSandboxCore::RenameSandboxWithDirectory(InNewName, OldSandboxName, BaseDirectory);
+	const bool bSuccess = FileSandboxCore::RenameSandboxWithDirectory(InNewName, OldSandboxName, BaseDirectory);
+
+	if (bSuccess)
+	{
+		NotifyIfBypassingRepository();
+	}
+
+	return bSuccess;
 }
 
 #define SET_ERROR(Text) if (OutError){ *OutError = Text; }
@@ -476,6 +547,21 @@ void FSandboxSystemModel::HandleActiveSandboxFilesChanged()
 	OnSandboxFilesChangedDelegate.Broadcast();
 }
 
+bool FSandboxSystemModel::IsUsingDefaultRepository() const
+{
+	const FString CustomBaseDir = SandboxModel::GetBaseSandboxDirectory();
+	const FString DefaultBaseDir = FileSandboxCore::GetBaseSandboxDirectory();
+	return CustomBaseDir == DefaultBaseDir;
+}
+
+void FSandboxSystemModel::NotifyIfBypassingRepository() const
+{
+	if (!IsUsingDefaultRepository())
+	{
+		OnKnownSandboxesChangedDelegate.Broadcast();
+	}
+}
+
 namespace SandboxModel
 {
 TOptional<FString> GetSandboxPathFor(const FString& InSandbox, const FString& InNonSandboxPath)
@@ -490,7 +576,27 @@ TOptional<FString> GetSandboxName(const FString& InSandboxDirectory)
 
 FString GetBaseSandboxDirectory()
 {
-	// For now we use the global default directory but this could change in the future, e.g. by settings.
+	// Priority: Command line > Custom setting > Default
+	// Check if settings specify a custom directory
+	const USandboxedEditingSettings* Settings = USandboxedEditingSettings::Get();
+	if (Settings && !Settings->CustomSandboxStorageDirectory.Path.IsEmpty())
+	{
+		// FileSandboxCore::GetBaseSandboxDirectory() handles command line override internally
+		// We need to ensure command line still wins, so check it first
+		const FString CoreResult = FileSandboxCore::GetBaseSandboxDirectory();
+		const FString DefaultIntermediate = FPaths::ProjectIntermediateDir() / TEXT("Sandboxes");
+
+		// If CoreResult != default, it means command line was used, so respect that
+		if (CoreResult != DefaultIntermediate)
+		{
+			return CoreResult;
+		}
+
+		// Otherwise use custom setting
+		return Settings->CustomSandboxStorageDirectory.Path;
+	}
+
+	// Fall back to FileSandboxCore (handles command line + default)
 	return FileSandboxCore::GetBaseSandboxDirectory();
 }
 
diff --git a/Engine/Plugins/Developer/Sandbox/SandboxedEditing/Source/SandboxedEditing/Private/Framework/Models/SandboxSystemModel.h b/Engine/Plugins/Developer/Sandbox/SandboxedEditing/Source/SandboxedEditing/Private/Framework/Models/SandboxSystemModel.h
--- a/Engine/Plugins/Developer/Sandbox/SandboxedEditing/Source/SandboxedEditing/Private/Framework/Models/SandboxSystemModel.h
+++ b/Engine/Plugins/Developer/Sandbox/SandboxedEditing/Source/SandboxedEditing/Private/Framework/Models/SandboxSystemModel.h
@@ -132,6 +132,12 @@ public:
 private:
 	/** @return true if a leave lock is currently registered. */
 	bool HasLeaveLock(FText* OutReason = nullptr) const;
+
+	/** @return true if using the default repository (not a custom directory). */
+	bool IsUsingDefaultRepository() const;
+
+	/** Fires OnKnownSandboxesChanged event if using a custom directory (bypassing repository's directory watching). */
+	void NotifyIfBypassingRepository() const;
 		
 	/** Invoked when the result of GetKnownSandboxes may have changed. */
 	FSimpleMulticastDelegate OnKnownSandboxesChangedDelegate;
diff --git a/Engine/Plugins/Developer/Sandbox/SandboxedEditing/Source/SandboxedEditing/Private/SandboxedEditingSettings.h b/Engine/Plugins/Developer/Sandbox/SandboxedEditing/Source/SandboxedEditing/Private/SandboxedEditingSettings.h
--- a/Engine/Plugins/Developer/Sandbox/SandboxedEditing/Source/SandboxedEditing/Private/SandboxedEditingSettings.h
+++ b/Engine/Plugins/Developer/Sandbox/SandboxedEditing/Source/SandboxedEditing/Private/SandboxedEditingSettings.h
@@ -2,7 +2,6 @@
 
 #pragma once
 
-#include "UObject/UObjectGlobals.h"
 #include "SandboxedEditingSettings.generated.h"
 
 /**
@@ -57,6 +56,10 @@ public:
 	/** Automatically delete sandboxes with no file changes when the editor shuts down */
 	UPROPERTY(config, EditAnywhere, Category = "Cleanup", meta = (DisplayName = "Delete Empty Sandboxes on Exit", ToolTip = "Automatically delete sandboxes with no file changes when the editor shuts down"))
 	bool bDeleteEmptySandboxesOnExit = false;
+
+	/** Optional custom directory where sandboxes are stored. Leave empty to use default location: {Project}/Intermediate/Sandboxes/. Requires editor restart to take effect. Can be overridden with -DEFAULTSANDBOXDIRECTORY command line argument. */
+	UPROPERTY(config, EditAnywhere, Category = "Storage", meta = (DisplayName = "Custom Sandbox Storage Directory", ToolTip = "Optional custom directory where sandboxes are stored. Leave empty to use default location: {Project}/Intermediate/Sandboxes/. Requires editor restart to take effect. Can be overridden with -DEFAULTSANDBOXDIRECTORY command line argument."))
+	FDirectoryPath CustomSandboxStorageDirectory;
 };
 
 
