Sandboxed Editing: Warn user about overwritting files when exporting multiple files.

#Synced-CL 52283655

diff --git a/Engine/Plugins/Developer/Sandbox/SandboxedEditing/Source/SandboxedEditing/Private/Features/Browser/ViewModels/Transfer/ExportWorkflow.cpp b/Engine/Plugins/Developer/Sandbox/SandboxedEditing/Source/SandboxedEditing/Private/Features/Browser/ViewModels/Transfer/ExportWorkflow.cpp
--- a/Engine/Plugins/Developer/Sandbox/SandboxedEditing/Source/SandboxedEditing/Private/Features/Browser/ViewModels/Transfer/ExportWorkflow.cpp
+++ b/Engine/Plugins/Developer/Sandbox/SandboxedEditing/Source/SandboxedEditing/Private/Features/Browser/ViewModels/Transfer/ExportWorkflow.cpp
@@ -127,6 +127,31 @@ void FExportWorkflow::Cancel()
 	OnWorkflowEndedDelegate.Execute();
 }
 
+TArray<FString> FExportWorkflow::GetFilesOverridenByExportDirectory(const FString& InNewExportRoot) const
+{
+	check(GetExportType() == EExportPathType::Directory);
+	TArray<FString> Result;
+	
+	for (const FString& SandboxRootPath : ExportedSandboxPaths)
+	{
+		const TOptional<FString> Name = SandboxModel::GetSandboxName(SandboxRootPath);
+		if (!Name)
+		{
+			continue;
+		}
+		
+		const FString BaseFile = FPaths::Combine(InNewExportRoot, *Name);
+		const FString WithExtension = FPaths::SetExtension(BaseFile, TEXT("zip"));
+		const bool bExists = FPaths::FileExists(WithExtension) || FPaths::DirectoryExists(WithExtension);
+		if (bExists)
+		{
+			Result.Add(WithExtension);
+		}
+	}
+	
+	return Result;
+}
+
 void FExportWorkflow::SetExportPath(FString InNewExportRoot)
 {
 	// For zip export, the system dialogue can specify a path or a full filename. But when exporting multiple items, only a directory is valid.
diff --git a/Engine/Plugins/Developer/Sandbox/SandboxedEditing/Source/SandboxedEditing/Private/Features/Browser/ViewModels/Transfer/ExportWorkflow.h b/Engine/Plugins/Developer/Sandbox/SandboxedEditing/Source/SandboxedEditing/Private/Features/Browser/ViewModels/Transfer/ExportWorkflow.h
--- a/Engine/Plugins/Developer/Sandbox/SandboxedEditing/Source/SandboxedEditing/Private/Features/Browser/ViewModels/Transfer/ExportWorkflow.h
+++ b/Engine/Plugins/Developer/Sandbox/SandboxedEditing/Source/SandboxedEditing/Private/Features/Browser/ViewModels/Transfer/ExportWorkflow.h
@@ -38,6 +38,9 @@ public:
 	/** Cancels the operation. */
 	void Cancel();
 	
+	/** @return List of files that would be overriden if this path is chosen. */
+	TArray<FString> GetFilesOverridenByExportDirectory(const FString& InNewExportRoot) const;
+	
 	/** Sets the export path. */
 	void SetExportPath(FString InNewExportRoot);
 	/** @return The export path */
diff --git a/Engine/Plugins/Developer/Sandbox/SandboxedEditing/Source/SandboxedEditing/Private/Features/Browser/Views/Transfer/ExportView.cpp b/Engine/Plugins/Developer/Sandbox/SandboxedEditing/Source/SandboxedEditing/Private/Features/Browser/Views/Transfer/ExportView.cpp
--- a/Engine/Plugins/Developer/Sandbox/SandboxedEditing/Source/SandboxedEditing/Private/Features/Browser/Views/Transfer/ExportView.cpp
+++ b/Engine/Plugins/Developer/Sandbox/SandboxedEditing/Source/SandboxedEditing/Private/Features/Browser/Views/Transfer/ExportView.cpp
@@ -7,6 +7,7 @@
 #include "Features/Browser/ViewModels/Transfer/ExportWorkflowResult.h"
 #include "Framework/Application/SlateApplication.h"
 #include "Framework/Notifications/NotificationManager.h"
+#include "Misc/MessageDialog.h"
 #include "Widgets/Notifications/SNotificationList.h"
 
 #define LOCTEXT_NAMESPACE "FExportView"
@@ -32,6 +33,58 @@ FExportView::~FExportView()
 
 namespace ExportViewDetail
 {
+static TOptional<FString> AskUserForZipPath(FExportWorkflow& InWorkflow, IDesktopPlatform& InDesktopPlatform)
+{
+	TArray<FString> SelectedFileNames;
+	const bool bSuccess = InDesktopPlatform.SaveFileDialog(
+		FSlateApplication::Get().FindBestParentWindowHandleForDialogs(nullptr),
+		LOCTEXT("ExportSandbox", "Export Sandbox").ToString(),
+		InWorkflow.GetDefaultExportDirectory(),
+		InWorkflow.GetDefaultSandboxFileName(),
+		TEXT("*.zip"),
+		EFileDialogFlags::None,
+		SelectedFileNames
+		);
+	return bSuccess && SelectedFileNames.Num() == 1 ? SelectedFileNames[0] : TOptional<FString>();
+}
+
+static TOptional<FString> AskUserForDirectoryPath(FExportWorkflow& InWorkflow, IDesktopPlatform& InDesktopPlatform)
+{
+	FString SelectedDirectory;
+	const bool bSuccess = InDesktopPlatform.OpenDirectoryDialog(
+		FSlateApplication::Get().FindBestParentWindowHandleForDialogs(nullptr),
+		LOCTEXT("ExportSandboxes", "Export Sandboxes").ToString(),
+		InWorkflow.GetDefaultExportDirectory(),
+		SelectedDirectory
+		);
+	if (!bSuccess)
+	{
+		return TOptional<FString>();
+	}
+	
+	const TArray<FString> OverridenFiles = InWorkflow.GetFilesOverridenByExportDirectory(SelectedDirectory);
+	if (OverridenFiles.IsEmpty())
+	{
+		return SelectedDirectory;
+	}
+	
+	const FString FilesList = FString::JoinBy(OverridenFiles, TEXT("\n- "), [](const FString& FilePath)
+	{
+		return FPaths::GetCleanFilename(FilePath);
+	});
+	const FText Message = FText::Format(
+		LOCTEXT("Directory.OverrideFiles.Message", "The following files in {0} will be overwritten:\n- {1}\n\nDo you want to continue?"), 
+		FText::AsCultureInvariant(SelectedDirectory),
+		FText::AsCultureInvariant(FilesList)
+		);
+	if (FMessageDialog::Open(EAppMsgType::OkCancel, Message, LOCTEXT("Directory.OverrideFiles.Title", "Override files")) == EAppReturnType::Ok)
+	{
+		return SelectedDirectory;
+	}
+	
+	return TOptional<FString>();
+}
+
 static TOptional<FString> AskUserForExportPath(FExportWorkflow& InWorkflow)
 {
 	IDesktopPlatform* DesktopPlatform = FDesktopPlatformModule::Get();
@@ -40,32 +93,13 @@ static TOptional<FString> AskUserForExportPath(FExportWorkflow& InWorkflow)
 		return {};
 	}
 	
-	const EExportPathType ExportType = InWorkflow.GetExportType();
 	static_assert(static_cast<int32>(EExportPathType::Count) == 2, "Update this logic.");
-	
-	if (ExportType == EExportPathType::Zip)
+	switch (InWorkflow.GetExportType())
 	{
-		TArray<FString> SelectedFileNames;
-		const bool bSuccess = DesktopPlatform->SaveFileDialog(
-			FSlateApplication::Get().FindBestParentWindowHandleForDialogs(nullptr),
-			LOCTEXT("ExportSandbox", "Export Sandbox").ToString(),
-			InWorkflow.GetDefaultExportDirectory(),
-			InWorkflow.GetDefaultSandboxFileName(),
-			TEXT("*.zip"),
-			EFileDialogFlags::None,
-			SelectedFileNames
-			);
-		return bSuccess && SelectedFileNames.Num() == 1 ? SelectedFileNames[0] : TOptional<FString>();
+	case EExportPathType::Directory: return AskUserForDirectoryPath(InWorkflow, *DesktopPlatform);
+	case EExportPathType::Zip: return AskUserForZipPath(InWorkflow, *DesktopPlatform);
+	default: checkNoEntry(); return {};
 	}
-	
-	FString SelectedDirectory;
-	const bool bSuccess = DesktopPlatform->OpenDirectoryDialog(
-		FSlateApplication::Get().FindBestParentWindowHandleForDialogs(nullptr),
-		LOCTEXT("ExportSandboxes", "Export Sandboxes").ToString(),
-		InWorkflow.GetDefaultExportDirectory(),
-		SelectedDirectory
-		); 
-	return bSuccess? SelectedDirectory : TOptional<FString>();
 }
 }
 
