Sandboxed Editing: Added Browse to Asset right-click option.

#Synced-CL 52378315

diff --git a/Engine/Plugins/Developer/Sandbox/SandboxedEditing/Source/SandboxedEditing/Private/Features/Browser/Commands/FileStateActions/FileStateCommandBindings.cpp b/Engine/Plugins/Developer/Sandbox/SandboxedEditing/Source/SandboxedEditing/Private/Features/Browser/Commands/FileStateActions/FileStateCommandBindings.cpp
--- a/Engine/Plugins/Developer/Sandbox/SandboxedEditing/Source/SandboxedEditing/Private/Features/Browser/Commands/FileStateActions/FileStateCommandBindings.cpp
+++ b/Engine/Plugins/Developer/Sandbox/SandboxedEditing/Source/SandboxedEditing/Private/Features/Browser/Commands/FileStateActions/FileStateCommandBindings.cpp
@@ -2,13 +2,18 @@
 
 #include "FileStateCommandBindings.h"
 
+#include "ContentBrowserModule.h"
 #include "FileStateCommands.h"
+#include "IContentBrowserSingleton.h"
+#include "AssetRegistry/AssetRegistryModule.h"
+#include "AssetRegistry/IAssetRegistry.h"
 #include "Features/Browser/Commands/BrowserCommands.h"
 #include "Features/Browser/ViewModels/FileState/FileStateItem.h"
 #include "Features/Browser/Views/Menu/FileStateMenuUtils.h"
 #include "Features/Browser/Widgets/Shared/FileState/SFileStateListView.h"
 #include "Framework/Models/SandboxSystemModel.h"
 #include "Framework/MultiBox/MultiBoxBuilder.h"
+#include "Misc/PackageName.h"
 
 namespace UE::SandboxedEditing
 {
@@ -44,6 +49,11 @@ void FFileStateCommandBindings::BindCommands()
 {
 	FFileStateCommands Commands = FFileStateCommands::Get();
 	
+	CommandList->MapAction(
+		Commands.BrowseToAsset, 
+		FExecuteAction::CreateRaw(this, &FFileStateCommandBindings::HandleBrowseToAsset), 
+		FCanExecuteAction::CreateRaw(this, &FFileStateCommandBindings::CanBrowseToAsset)
+		);
 	CommandList->MapAction(
 		Commands.ShowRootInExplorer, FExecuteAction::CreateRaw(this, &FFileStateCommandBindings::HandleShowRootInExplorer)
 		);
@@ -54,6 +64,44 @@ void FFileStateCommandBindings::BindCommands()
 		);
 }
 
+namespace CommandDetails
+{
+static TArray<FAssetData> GetAssetDataFromPaths(const TArray<FString>& InAbsolutePaths)                                                                                                                                                                                                                          
+{                                                                                                                                                                                                                                                                                                                
+	IAssetRegistry& AssetRegistry = FModuleManager::LoadModuleChecked<FAssetRegistryModule>(AssetRegistryConstants::ModuleName).Get();                                                                                                                                                                           
+	
+	TArray<FAssetData> Result;                                                                                                                                                                                                                                                                                   
+	Result.Reserve(InAbsolutePaths.Num());                                                                                                                                                                                                                                                                       
+	
+	for (const FString& AbsolutePath : InAbsolutePaths)
+	{
+		FString PackageName;
+		if (FPackageName::TryConvertFilenameToLongPackageName(AbsolutePath, PackageName))
+		{
+			AssetRegistry.GetAssetsByPackageName(FName(*PackageName), Result);
+		}
+  }
+
+  return Result;
+}
+}
+
+void FFileStateCommandBindings::HandleBrowseToAsset() const
+{
+	const TArray<FAssetData> Assets = CommandDetails::GetAssetDataFromPaths(SelectedFilePathsAttr.Get());
+	if (!Assets.IsEmpty())
+	{
+		const FContentBrowserModule& ContentBrowserModule = FModuleManager::Get().LoadModuleChecked<FContentBrowserModule>("ContentBrowser");
+		ContentBrowserModule.Get().SyncBrowserToAssets(Assets);
+	}
+}
+
+bool FFileStateCommandBindings::CanBrowseToAsset() const
+{
+	const TArray<FAssetData> Assets = CommandDetails::GetAssetDataFromPaths(SelectedFilePathsAttr.Get());
+	return !Assets.IsEmpty();
+}
+
 void FFileStateCommandBindings::HandleShowRootInExplorer() const
 {
 	if (const TOptional<FString> SandboxRoot = SelectedSandboxRootAttr.Get())
diff --git a/Engine/Plugins/Developer/Sandbox/SandboxedEditing/Source/SandboxedEditing/Private/Features/Browser/Commands/FileStateActions/FileStateCommandBindings.h b/Engine/Plugins/Developer/Sandbox/SandboxedEditing/Source/SandboxedEditing/Private/Features/Browser/Commands/FileStateActions/FileStateCommandBindings.h
--- a/Engine/Plugins/Developer/Sandbox/SandboxedEditing/Source/SandboxedEditing/Private/Features/Browser/Commands/FileStateActions/FileStateCommandBindings.h
+++ b/Engine/Plugins/Developer/Sandbox/SandboxedEditing/Source/SandboxedEditing/Private/Features/Browser/Commands/FileStateActions/FileStateCommandBindings.h
@@ -43,6 +43,9 @@ private:
 	/** Binds the commands needed for the menu. */
 	void BindCommands();
 	
+	void HandleBrowseToAsset() const;
+	bool CanBrowseToAsset() const;
+	
 	void HandleShowRootInExplorer() const;
 	
 	void HandleShowFileInExplorer() const;
diff --git a/Engine/Plugins/Developer/Sandbox/SandboxedEditing/Source/SandboxedEditing/Private/Features/Browser/Commands/FileStateActions/FileStateCommands.cpp b/Engine/Plugins/Developer/Sandbox/SandboxedEditing/Source/SandboxedEditing/Private/Features/Browser/Commands/FileStateActions/FileStateCommands.cpp
--- a/Engine/Plugins/Developer/Sandbox/SandboxedEditing/Source/SandboxedEditing/Private/Features/Browser/Commands/FileStateActions/FileStateCommands.cpp
+++ b/Engine/Plugins/Developer/Sandbox/SandboxedEditing/Source/SandboxedEditing/Private/Features/Browser/Commands/FileStateActions/FileStateCommands.cpp
@@ -19,6 +19,7 @@ FFileStateCommands::FFileStateCommands()
 
 void FFileStateCommands::RegisterCommands()
 {
+	UI_COMMAND(BrowseToAsset, "Browse to asset", "Shows this asset in the content browser.", EUserInterfaceActionType::Button, FInputChord())
 	UI_COMMAND(ShowRootInExplorer, "Show sandbox folder", "Shows sandbox root directory in the system file manager.", EUserInterfaceActionType::Button, FInputChord())
 	UI_COMMAND(ShowFileInExplorer, "Show in file manager", "Shows the file's sandboxed counterpart in the system file manager.", EUserInterfaceActionType::Button, FInputChord())
 	
diff --git a/Engine/Plugins/Developer/Sandbox/SandboxedEditing/Source/SandboxedEditing/Private/Features/Browser/Commands/FileStateActions/FileStateCommands.h b/Engine/Plugins/Developer/Sandbox/SandboxedEditing/Source/SandboxedEditing/Private/Features/Browser/Commands/FileStateActions/FileStateCommands.h
--- a/Engine/Plugins/Developer/Sandbox/SandboxedEditing/Source/SandboxedEditing/Private/Features/Browser/Commands/FileStateActions/FileStateCommands.h
+++ b/Engine/Plugins/Developer/Sandbox/SandboxedEditing/Source/SandboxedEditing/Private/Features/Browser/Commands/FileStateActions/FileStateCommands.h
@@ -17,9 +17,10 @@ public:
 	virtual void RegisterCommands() override;
 	//~ End TCommands Interface
 	
+	/** Opens the content browser to this file*/
+	TSharedPtr<FUICommandInfo> BrowseToAsset;
 	/** Shows sandbox root directory in the system file explorer. */
 	TSharedPtr<FUICommandInfo> ShowRootInExplorer;
-	
 	/** Shows the file's sandboxed counterpart in the system file explorer. */
 	TSharedPtr<FUICommandInfo> ShowFileInExplorer;
 	
diff --git a/Engine/Plugins/Developer/Sandbox/SandboxedEditing/Source/SandboxedEditing/Private/Features/Browser/Views/Menu/FileStateMenuUtils.cpp b/Engine/Plugins/Developer/Sandbox/SandboxedEditing/Source/SandboxedEditing/Private/Features/Browser/Views/Menu/FileStateMenuUtils.cpp
--- a/Engine/Plugins/Developer/Sandbox/SandboxedEditing/Source/SandboxedEditing/Private/Features/Browser/Views/Menu/FileStateMenuUtils.cpp
+++ b/Engine/Plugins/Developer/Sandbox/SandboxedEditing/Source/SandboxedEditing/Private/Features/Browser/Views/Menu/FileStateMenuUtils.cpp
@@ -49,6 +49,7 @@ void AppendMenu(FMenuBuilder& InMenuBuilder)
 	InMenuBuilder.BeginSection(TEXT("Explorer"), LOCTEXT("Explorer", "Explorer"));
 	
 	FFileStateCommands& Commands = FFileStateCommands::Get();
+	InMenuBuilder.AddMenuEntry(Commands.BrowseToAsset);
 	InMenuBuilder.AddMenuEntry(Commands.ShowRootInExplorer);
 	InMenuBuilder.AddMenuEntry(Commands.ShowFileInExplorer, NAME_None, Detail::GetExploreFolderText());
 	
diff --git a/Engine/Plugins/Developer/Sandbox/SandboxedEditing/Source/SandboxedEditing/Private/SandboxedEditingStyle.cpp b/Engine/Plugins/Developer/Sandbox/SandboxedEditing/Source/SandboxedEditing/Private/SandboxedEditingStyle.cpp
--- a/Engine/Plugins/Developer/Sandbox/SandboxedEditing/Source/SandboxedEditing/Private/SandboxedEditingStyle.cpp
+++ b/Engine/Plugins/Developer/Sandbox/SandboxedEditing/Source/SandboxedEditing/Private/SandboxedEditingStyle.cpp
@@ -50,6 +50,7 @@ void FSandboxedEditingStyle::RegisterBrowserStyle()
 	Set("Browser.ExportSandboxes", new FSlateBrush(*FAppStyle::Get().GetBrush("Icons.Export")));
 	Set("Browser.ImportSandboxes", new FSlateBrush(*FAppStyle::Get().GetBrush("Icons.Import")));
 	
+	Set("FileState.BrowseToAsset", new FSlateBrush(*FAppStyle::Get().GetBrush("SystemWideCommands.FindInContentBrowser")));
 	Set("FileState.ShowRootInExplorer", new FSlateBrush(*FAppStyle::Get().GetBrush("SystemWideCommands.FindInContentBrowser")));
 	Set("FileState.ShowFileInExplorer", new FSlateBrush(*FAppStyle::Get().GetBrush("SystemWideCommands.FindInContentBrowser")));
 	
