Add sandbox description to the sandbox browser list.

#Synced-CL 52444341

diff --git a/Engine/Plugins/Developer/Sandbox/SandboxedEditing/Source/SandboxedEditing/Private/Features/Browser/BrowserFeature.cpp b/Engine/Plugins/Developer/Sandbox/SandboxedEditing/Source/SandboxedEditing/Private/Features/Browser/BrowserFeature.cpp
--- a/Engine/Plugins/Developer/Sandbox/SandboxedEditing/Source/SandboxedEditing/Private/Features/Browser/BrowserFeature.cpp
+++ b/Engine/Plugins/Developer/Sandbox/SandboxedEditing/Source/SandboxedEditing/Private/Features/Browser/BrowserFeature.cpp
@@ -8,6 +8,7 @@
 #include "ViewModels/SandboxControlsViewModel.h"
 #include "Widgets/Docking/SDockTab.h"
 #include "Widgets/SSandboxedEditingRoot.h"
+#include "Widgets/Unsandboxed/List/Columns/DescriptionSandboxColumn.h"
 #include "Widgets/Unsandboxed/List/Columns/LastModifiedSandboxColumn.h"
 #include "Widgets/Unsandboxed/List/Columns/NameSandboxColumn.h"
 #include "Widgets/Unsandboxed/List/Columns/VersionSandboxColumn.h"
@@ -28,17 +29,20 @@ static void MakeSandboxColumns(
 	)
 {
 	const TSharedRef<FNameSandboxColumn> Name = MakeShared<FNameSandboxColumn>(InControlsViewModel);
+	const TSharedRef<FDescriptionSandboxColumn> Description = MakeShared<FDescriptionSandboxColumn>();
 	const TSharedRef<FVersionSandboxColumn> Version = MakeShared<FVersionSandboxColumn>();
 	const TSharedRef<FLastModifiedSandboxColumn> LastModified = MakeShared<FLastModifiedSandboxColumn>();
 	OutBehaviors = TMap<FName, TSharedRef<ISandboxColumnBehavior>>
 	{
 		{ NameSandboxColumn, Name },
+		{ DescriptionSandboxColumn, Description },
 		{ VersionSandboxNameColumn, Version },
 		{ LastModifiedSandboxColumn, LastModified }
 	};
 	OutFactories = TMap<FName, TSharedRef<ISandboxColumnWidgetFactory>>
 	{
 		{ NameSandboxColumn, Name },
+		{ DescriptionSandboxColumn, Description },
 		{ VersionSandboxNameColumn, Version },
 		{ LastModifiedSandboxColumn, LastModified }
 	};
diff --git a/Engine/Plugins/Developer/Sandbox/SandboxedEditing/Source/SandboxedEditing/Private/Features/Browser/ViewModels/List/SandboxColumns.h b/Engine/Plugins/Developer/Sandbox/SandboxedEditing/Source/SandboxedEditing/Private/Features/Browser/ViewModels/List/SandboxColumns.h
--- a/Engine/Plugins/Developer/Sandbox/SandboxedEditing/Source/SandboxedEditing/Private/Features/Browser/ViewModels/List/SandboxColumns.h
+++ b/Engine/Plugins/Developer/Sandbox/SandboxedEditing/Source/SandboxedEditing/Private/Features/Browser/ViewModels/List/SandboxColumns.h
@@ -15,6 +15,7 @@ class ISandboxColumnWidgetFactory;
 using FSandboxColumnFactoryMap = TMap<FName, TSharedRef<ISandboxColumnWidgetFactory>>;
 
 constexpr FLazyName NameSandboxColumn(TEXT("NameSandboxColumn"));
+constexpr FLazyName DescriptionSandboxColumn(TEXT("DescriptionSandboxColumn"));
 constexpr FLazyName VersionSandboxNameColumn(TEXT("VersionSandboxNameColumn"));
 constexpr FLazyName LastModifiedSandboxColumn(TEXT("LastModifiedSandboxColumn"));
 }
diff --git a/Engine/Plugins/Developer/Sandbox/SandboxedEditing/Source/SandboxedEditing/Private/Features/Browser/ViewModels/List/SandboxListViewModel.cpp b/Engine/Plugins/Developer/Sandbox/SandboxedEditing/Source/SandboxedEditing/Private/Features/Browser/ViewModels/List/SandboxListViewModel.cpp
--- a/Engine/Plugins/Developer/Sandbox/SandboxedEditing/Source/SandboxedEditing/Private/Features/Browser/ViewModels/List/SandboxListViewModel.cpp
+++ b/Engine/Plugins/Developer/Sandbox/SandboxedEditing/Source/SandboxedEditing/Private/Features/Browser/ViewModels/List/SandboxListViewModel.cpp
@@ -3,6 +3,7 @@
 #include "SandboxListViewModel.h"
 
 #include "FilterSandboxViewModel.h"
+#include "SandboxColumns.h"
 #include "SandboxListItem.h"
 #include "Templates/SharedPointer.h"
 
@@ -32,12 +33,48 @@ FSandboxListViewModel::~FSandboxListViewModel()
 
 TArray<FName> FSandboxListViewModel::GetDisplayedColumns() const
 {
-	// For now all columns are visible. In the future, we may allow the user to hide columns.
 	TArray<FName> ColumnIds;
-	Columns.GenerateKeyArray(ColumnIds);
+	for (const auto& Pair : Columns)
+	{
+		if (!HiddenColumns.Contains(Pair.Key))
+		{
+			ColumnIds.Add(Pair.Key);
+		}
+	}
 	return ColumnIds;
 }
 
+bool FSandboxListViewModel::IsColumnVisible(FName InColumnName) const
+{
+	return !HiddenColumns.Contains(InColumnName);
+}
+
+void FSandboxListViewModel::SetColumnVisibility(FName InColumnName, bool bIsVisible)
+{
+	// Name column cannot be hidden
+	if (InColumnName == NameSandboxColumn)
+	{
+		return;
+	}
+
+	const bool bIsCurrentlyVisible = !HiddenColumns.Contains(InColumnName);
+	if (bIsVisible == bIsCurrentlyVisible)
+	{
+		return;
+	}
+
+	if (bIsVisible)
+	{
+		HiddenColumns.Remove(InColumnName);
+	}
+	else
+	{
+		HiddenColumns.Add(InColumnName);
+	}
+
+	OnItemsChangedDelegate.Broadcast();
+}
+
 EColumnSortPriority::Type FSandboxListViewModel::GetSortPriority(FName InColumnName) const
 {
 	// TODO DP:
diff --git a/Engine/Plugins/Developer/Sandbox/SandboxedEditing/Source/SandboxedEditing/Private/Features/Browser/ViewModels/List/SandboxListViewModel.h b/Engine/Plugins/Developer/Sandbox/SandboxedEditing/Source/SandboxedEditing/Private/Features/Browser/ViewModels/List/SandboxListViewModel.h
--- a/Engine/Plugins/Developer/Sandbox/SandboxedEditing/Source/SandboxedEditing/Private/Features/Browser/ViewModels/List/SandboxListViewModel.h
+++ b/Engine/Plugins/Developer/Sandbox/SandboxedEditing/Source/SandboxedEditing/Private/Features/Browser/ViewModels/List/SandboxListViewModel.h
@@ -26,9 +26,15 @@ public:
 	
 	/** @return The items to display */
 	const TArray<TSharedPtr<FSandboxListItem>>& GetItems() const { return FilteredItems; }
-	
-	/** @return The columns that should be displayed. */
+
+	/** @return The columns that should be displayed (taking visibility into account). */
 	TArray<FName> GetDisplayedColumns() const;
+
+	/** @return Whether a column is visible. Name column is always visible. */
+	bool IsColumnVisible(FName InColumnName) const;
+
+	/** Sets whether a column is visible. Name column cannot be hidden. */
+	void SetColumnVisibility(FName InColumnName, bool bIsVisible);
 	
 	/** @return The priority with which the column is sorted. */
 	EColumnSortPriority::Type GetSortPriority(FName InColumnName) const;
@@ -49,7 +55,10 @@ private:
 	
 	/** All columns that can ever be displayed. */
 	const TMap<FName, TSharedRef<ISandboxColumnBehavior>> Columns;
-	
+
+	/** Tracks which columns are visible. Name column is always visible. */
+	TSet<FName> HiddenColumns;
+
 	/** Contains all items. */
 	TArray<TSharedPtr<FSandboxListItem>> AllUnfilteredItems;
 	/** A filtered version of AllUnfilteredItems. */
diff --git a/Engine/Plugins/Developer/Sandbox/SandboxedEditing/Source/SandboxedEditing/Private/Features/Browser/Widgets/Unsandboxed/List/Columns/DescriptionSandboxColumn.cpp b/Engine/Plugins/Developer/Sandbox/SandboxedEditing/Source/SandboxedEditing/Private/Features/Browser/Widgets/Unsandboxed/List/Columns/DescriptionSandboxColumn.cpp
new file mode 100644
--- /dev/null
+++ b/Engine/Plugins/Developer/Sandbox/SandboxedEditing/Source/SandboxedEditing/Private/Features/Browser/Widgets/Unsandboxed/List/Columns/DescriptionSandboxColumn.cpp
@@ -0,0 +1,55 @@
+// Copyright Epic Games, Inc. All Rights Reserved.
+
+#include "DescriptionSandboxColumn.h"
+
+#include "Features/Browser/ViewModels/List/SandboxColumns.h"
+#include "Features/Browser/ViewModels/List/SandboxListItem.h"
+#include "Framework/Models/SandboxInfo.h"
+#include "SandboxedEditingStyle.h"
+#include "Widgets/Layout/SBox.h"
+#include "Widgets/Text/STextBlock.h"
+#include "Widgets/Views/SHeaderRow.h"
+
+#define LOCTEXT_NAMESPACE "FDescriptionSandboxColumn"
+
+namespace UE::SandboxedEditing
+{
+namespace DescriptionDetail
+{
+static FText FormatDescription(const FSandboxInfo& InRowData)
+{
+	return FText::AsCultureInvariant(InRowData.Description);
+}
+}
+
+void FDescriptionSandboxColumn::PopulateSearchTerms(const TSharedPtr<FSandboxListItem>& InRowData, TArray<FString>& OutSearchTerms) const
+{
+	OutSearchTerms.Add(
+		DescriptionDetail::FormatDescription(InRowData->SandboxInfo).ToString()
+		);
+}
+
+SHeaderRow::FColumn::FArguments FDescriptionSandboxColumn::MakeColumnArguments()
+{
+	return SHeaderRow::FColumn::FArguments()
+		.ColumnId(DescriptionSandboxColumn)
+		.FillWidth(FSandboxedEditingStyle::Get().GetFloat("SandboxedEditing.Browser.DescriptionColumn.FillWidth"))
+		.DefaultLabel(LOCTEXT("Description.Label", "Description"))
+		.ToolTipText(LOCTEXT("Description.ToolTip", "The sandbox description"));
+}
+
+TSharedRef<SWidget> FDescriptionSandboxColumn::MakeColumnWidget(const FMakeSandboxColumnWidgetArgs& InArgs)
+{
+	return SNew(SBox)
+		.VAlign(VAlign_Center)
+		.Padding(FSandboxedEditingStyle::Get().GetMargin("SandboxedEditing.Browser.DescriptionColumn.Padding"))
+		[
+			SNew(STextBlock)
+			.Text_Lambda([RowData = InArgs.RowData](){ return DescriptionDetail::FormatDescription(RowData->SandboxInfo); })
+			.HighlightText(InArgs.HighlightText)
+			.OverflowPolicy(ETextOverflowPolicy::Ellipsis)
+		];
+}
+}
+
+#undef LOCTEXT_NAMESPACE
diff --git a/Engine/Plugins/Developer/Sandbox/SandboxedEditing/Source/SandboxedEditing/Private/Features/Browser/Widgets/Unsandboxed/List/Columns/DescriptionSandboxColumn.h b/Engine/Plugins/Developer/Sandbox/SandboxedEditing/Source/SandboxedEditing/Private/Features/Browser/Widgets/Unsandboxed/List/Columns/DescriptionSandboxColumn.h
new file mode 100644
--- /dev/null
+++ b/Engine/Plugins/Developer/Sandbox/SandboxedEditing/Source/SandboxedEditing/Private/Features/Browser/Widgets/Unsandboxed/List/Columns/DescriptionSandboxColumn.h
@@ -0,0 +1,24 @@
+// Copyright Epic Games, Inc. All Rights Reserved.
+
+#pragma once
+
+#include "Features/Browser/ViewModels/List/ISandboxColumnBehavior.h"
+#include "Features/Browser/ViewModels/List/ISandboxColumnWidgetFactory.h"
+
+namespace UE::SandboxedEditing
+{
+/** Displays the sandbox description. */
+class FDescriptionSandboxColumn : public ISandboxColumnBehavior, public ISandboxColumnWidgetFactory
+{
+public:
+
+	//~ Begin ISandboxColumnBehavior Interface
+	virtual void PopulateSearchTerms(const TSharedPtr<FSandboxListItem>& InRowData, TArray<FString>& OutSearchTerms) const override;
+	//~ End ISandboxColumnBehavior Interface
+
+	//~ Begin ISandboxWidgetFactory Interface
+	virtual SHeaderRow::FColumn::FArguments MakeColumnArguments() override;
+	virtual TSharedRef<SWidget> MakeColumnWidget(const FMakeSandboxColumnWidgetArgs& InArgs) override;
+	//~ End ISandboxWidgetFactory Interface
+};
+}
diff --git a/Engine/Plugins/Developer/Sandbox/SandboxedEditing/Source/SandboxedEditing/Private/Features/Browser/Widgets/Unsandboxed/List/SSandboxListView.cpp b/Engine/Plugins/Developer/Sandbox/SandboxedEditing/Source/SandboxedEditing/Private/Features/Browser/Widgets/Unsandboxed/List/SSandboxListView.cpp
--- a/Engine/Plugins/Developer/Sandbox/SandboxedEditing/Source/SandboxedEditing/Private/Features/Browser/Widgets/Unsandboxed/List/SSandboxListView.cpp
+++ b/Engine/Plugins/Developer/Sandbox/SandboxedEditing/Source/SandboxedEditing/Private/Features/Browser/Widgets/Unsandboxed/List/SSandboxListView.cpp
@@ -3,7 +3,9 @@
 #include "SSandboxListView.h"
 
 #include "Features/Browser/Commands/BrowserCommands.h"
+#include "Features/Browser/ViewModels/List/ISandboxColumnBehavior.h"
 #include "Features/Browser/ViewModels/List/ISandboxColumnWidgetFactory.h"
+#include "Features/Browser/ViewModels/List/SandboxColumns.h"
 #include "Features/Browser/ViewModels/List/SandboxListItem.h"
 #include "Features/Browser/ViewModels/List/SandboxListViewModel.h"
 #include "Features/Browser/ViewModels/SandboxControlsViewModel.h"
@@ -13,6 +15,8 @@
 #include "SSandboxListRow.h"
 #include "SandboxedEditingStyle.h"
 #include "Utils/BreakBehavior.h"
+#include "Widgets/Input/SCheckBox.h"
+#include "Widgets/Input/SComboButton.h"
 #include "Widgets/Views/SListView.h"
 #include "Widgets/SBoxPanel.h"
 #include "Widgets/SOverlay.h"
@@ -112,8 +116,9 @@ int32 SSandboxListView::NumSelected() const
 
 TSharedRef<SHeaderRow> SSandboxListView::MakeHeaderRow()
 {
-	const TSharedRef<SHeaderRow> HeaderRow = SNew(SHeaderRow);
-		
+	const TSharedRef<SHeaderRow> HeaderRow = SNew(SHeaderRow)
+		.CanSelectGeneratedColumn(true);
+
 	for (FName ColumnId : ListViewModel->GetDisplayedColumns())
 	{
 		if (const TSharedRef<ISandboxColumnWidgetFactory>* Factory = ColumnFactoryMap.Find(ColumnId))
@@ -121,10 +126,55 @@ TSharedRef<SHeaderRow> SSandboxListView::MakeHeaderRow()
 			HeaderRow->AddColumn(Factory->Get().MakeColumnArguments());
 		}
 	}
-	
+
 	return HeaderRow;
 }
 
+TSharedRef<SWidget> SSandboxListView::MakeColumnVisibilityMenu()
+{
+	FMenuBuilder MenuBuilder(true, nullptr);
+
+	MenuBuilder.BeginSection(TEXT("Columns"), LOCTEXT("Columns", "Columns"));
+	{
+		// Add checkbox for each column
+		for (const auto& Pair : ColumnFactoryMap)
+		{
+			const FName ColumnId = Pair.Key;
+
+			// Get column label
+			SHeaderRow::FColumn::FArguments ColumnArgs = Pair.Value->MakeColumnArguments();
+			const FText ColumnLabel = ColumnArgs._DefaultLabel.Get();
+
+			MenuBuilder.AddMenuEntry(
+				ColumnLabel,
+				FText::GetEmpty(),
+				FSlateIcon(),
+				FUIAction(
+					FExecuteAction::CreateLambda([this, ColumnId]()
+					{
+						const bool bIsVisible = ListViewModel->IsColumnVisible(ColumnId);
+						ListViewModel->SetColumnVisibility(ColumnId, !bIsVisible);
+					}),
+					FCanExecuteAction::CreateLambda([ColumnId]()
+					{
+						// Name column cannot be hidden
+						return ColumnId != NameSandboxColumn;
+					}),
+					FIsActionChecked::CreateLambda([this, ColumnId]()
+					{
+						return ListViewModel->IsColumnVisible(ColumnId);
+					})
+				),
+				NAME_None,
+				EUserInterfaceActionType::ToggleButton
+			);
+		}
+	}
+	MenuBuilder.EndSection();
+
+	return MenuBuilder.MakeWidget();
+}
+
 TSharedRef<ITableRow> SSandboxListView::OnGenerateRow(TSharedPtr<FSandboxListItem> InItem, const TSharedRef<STableViewBase>& InOwnerTable)
 {
 	check(InItem);
@@ -237,4 +287,4 @@ EVisibility SSandboxListView::GetOverlayTextVisibility() const
 }
 }
 
-#undef LOCTEXT_NAMESPACE
\ No newline at end of file
+#undef LOCTEXT_NAMESPACE
diff --git a/Engine/Plugins/Developer/Sandbox/SandboxedEditing/Source/SandboxedEditing/Private/Features/Browser/Widgets/Unsandboxed/List/SSandboxListView.h b/Engine/Plugins/Developer/Sandbox/SandboxedEditing/Source/SandboxedEditing/Private/Features/Browser/Widgets/Unsandboxed/List/SSandboxListView.h
--- a/Engine/Plugins/Developer/Sandbox/SandboxedEditing/Source/SandboxedEditing/Private/Features/Browser/Widgets/Unsandboxed/List/SSandboxListView.h
+++ b/Engine/Plugins/Developer/Sandbox/SandboxedEditing/Source/SandboxedEditing/Private/Features/Browser/Widgets/Unsandboxed/List/SSandboxListView.h
@@ -91,7 +91,10 @@ private:
 	
 	/** @return The header row to display at the top of the table */
 	TSharedRef<SHeaderRow> MakeHeaderRow();
-	
+
+	/** @return Column visibility menu content */
+	TSharedRef<SWidget> MakeColumnVisibilityMenu();
+
 	/** @return Widget for the items */
 	TSharedRef<ITableRow> OnGenerateRow(TSharedPtr<FSandboxListItem> InItem, const TSharedRef<STableViewBase>& InOwnerTable);
 	/** Loads the clicked sandbox */
diff --git a/Engine/Plugins/Developer/Sandbox/SandboxedEditing/Source/SandboxedEditing/Private/SandboxedEditingStyle.cpp b/Engine/Plugins/Developer/Sandbox/SandboxedEditing/Source/SandboxedEditing/Private/SandboxedEditingStyle.cpp
--- a/Engine/Plugins/Developer/Sandbox/SandboxedEditing/Source/SandboxedEditing/Private/SandboxedEditingStyle.cpp
+++ b/Engine/Plugins/Developer/Sandbox/SandboxedEditing/Source/SandboxedEditing/Private/SandboxedEditingStyle.cpp
@@ -58,6 +58,8 @@ void FSandboxedEditingStyle::RegisterBrowserStyle()
 	Set("SandboxedEditing.Browser.RowPadding", FMargin(0.f, 4.f));
 	Set("SandboxedEditing.Browser.NameColumn.Padding", FMargin(16.f, 0.f, 0.f, 0.f));
 	Set("SandboxedEditing.Browser.NameColumn.FillWidth", 2.f);
+	Set("SandboxedEditing.Browser.DescriptionColumn.Padding", FMargin(8.f, 0.f, 0.f, 0.f));
+	Set("SandboxedEditing.Browser.DescriptionColumn.FillWidth", 3.f);
 	Set("SandboxedEditing.Browser.VersionColumn.Padding", FMargin(8.f, 0.f, 0.f, 0.f));
 	Set("SandboxedEditing.Browser.VersionColumn.FillWidth", 1.f);
 	Set("SandboxedEditing.Browser.LastModifiedColumn.Padding", FMargin(8.f, 0.f, 0.f, 0.f));
