Sandboxed Editing: Last Modified shows wrong time due to FDateTime::AsDate performing unwanted conversion.

#Synced-CL 52378315

diff --git a/Engine/Plugins/Developer/Sandbox/SandboxedEditing/Source/SandboxedEditing/Private/Features/Browser/Widgets/Unsandboxed/List/Columns/LastModifiedSandboxColumn.cpp b/Engine/Plugins/Developer/Sandbox/SandboxedEditing/Source/SandboxedEditing/Private/Features/Browser/Widgets/Unsandboxed/List/Columns/LastModifiedSandboxColumn.cpp
--- a/Engine/Plugins/Developer/Sandbox/SandboxedEditing/Source/SandboxedEditing/Private/Features/Browser/Widgets/Unsandboxed/List/Columns/LastModifiedSandboxColumn.cpp
+++ b/Engine/Plugins/Developer/Sandbox/SandboxedEditing/Source/SandboxedEditing/Private/Features/Browser/Widgets/Unsandboxed/List/Columns/LastModifiedSandboxColumn.cpp
@@ -18,7 +18,8 @@ namespace LastModifiedDetail
 {
 static FText FormatLastModified(const FSandboxInfo& InRowData)
 {
-	return FText::AsDateTime(InRowData.LastModified);
+	// InRowData.LastModified is already in local time, hence must pass in FText::GetInvariantTimeZone() to avoid double conversion.
+	return FText::AsDateTime(InRowData.LastModified, EDateTimeStyle::Default, EDateTimeStyle::Default, FText::GetInvariantTimeZone());
 }
 }
 void FLastModifiedSandboxColumn::PopulateSearchTerms(const TSharedPtr<FSandboxListItem>& InRowData, TArray<FString>& OutSearchTerms) const
diff --git a/Engine/Plugins/Developer/Sandbox/SandboxedEditing/Source/SandboxedEditing/Private/Framework/Models/SandboxInfo.h b/Engine/Plugins/Developer/Sandbox/SandboxedEditing/Source/SandboxedEditing/Private/Framework/Models/SandboxInfo.h
--- a/Engine/Plugins/Developer/Sandbox/SandboxedEditing/Source/SandboxedEditing/Private/Framework/Models/SandboxInfo.h
+++ b/Engine/Plugins/Developer/Sandbox/SandboxedEditing/Source/SandboxedEditing/Private/Framework/Models/SandboxInfo.h
@@ -20,7 +20,7 @@ struct FSandboxInfo
 	/** Path to the root of the sandbox directory. Unique for every sandbox item. */
 	FString SandboxRoot;
 	
-	/** Info about when the sandbox was last modified. */
+	/** Info about when the sandbox was last modified in local time. */
 	FDateTime LastModified;
 	
 	/** The version of the engine when this sandbox was created in. */
