Add a toast to warn users about potential loss of current changes if they join a sandbox.

#rnx

#Synced-CL 52444341

diff --git a/Engine/Plugins/Developer/Sandbox/SandboxedEditing/Source/SandboxedEditing/Private/Features/Browser/ViewModels/SandboxControlsViewModel.cpp b/Engine/Plugins/Developer/Sandbox/SandboxedEditing/Source/SandboxedEditing/Private/Features/Browser/ViewModels/SandboxControlsViewModel.cpp
--- a/Engine/Plugins/Developer/Sandbox/SandboxedEditing/Source/SandboxedEditing/Private/Features/Browser/ViewModels/SandboxControlsViewModel.cpp
+++ b/Engine/Plugins/Developer/Sandbox/SandboxedEditing/Source/SandboxedEditing/Private/Features/Browser/ViewModels/SandboxControlsViewModel.cpp
@@ -2,6 +2,10 @@
 
 #include "SandboxControlsViewModel.h"
 
+#include "Framework/Models/SandboxInfo.h"
+#include "Framework/Models/SandboxSystemModel.h"
+#include "Framework/Notifications.h"
+
 namespace UE::SandboxedEditing
 {
 FSandboxControlsViewModel::FSandboxControlsViewModel(const TSharedRef<FSandboxSystemModel>& InModel)
@@ -49,12 +53,58 @@ void FSandboxControlsViewModel::LoadSandbox(const FString& InSandboxRoot)
 	{
 		return;
 	}
+
+	if (InMemoryChangesNotification.IsValid())
+	{
+		// We are actively waiting for a user to make a decision on entry into the sandbox. 
+		return;
+	}
 	
+	// Check for in-memory changes before loading
+	const TArray<UPackage*> DirtyPackages = Model->GetInMemoryChanges();
+	if (!DirtyPackages.IsEmpty())
+	{
+		if (CurrentCreationWorkflow)
+		{
+			CurrentCreationWorkflow->Cancel();
+		}
+
+		// Get sandbox name for display
+		const TOptional<FSandboxInfo> SandboxInfo = Model->GetSandboxInfo(InSandboxRoot);
+		const FString SandboxName = SandboxInfo.IsSet() ? SandboxInfo->Name : TEXT("Sandbox");
+
+		TWeakPtr<FSandboxSystemModel> WeakModel = Model; 
+		// Show warning notification with Discard/Cancel options
+		InMemoryChangesNotification = ShowInMemoryChangesWarning(
+			SandboxName,
+			// OnDiscard: Discard changes and load sandbox
+			[WeakModel, InSandboxRoot]()
+			{
+				if (TSharedPtr<FSandboxSystemModel> Model = WeakModel.Pin())
+				{
+					Model->RevertDirtyPackages();
+					Model->LoadSandbox(InSandboxRoot);
+				}
+			},
+			// OnCancel: Do nothing
+			[]()
+			{
+				// User cancelled, do nothing
+			}
+		);
+		return;
+	}
+
+	LoadSandboxInternal(InSandboxRoot);
+}
+
+void FSandboxControlsViewModel::LoadSandboxInternal(const FString& InSandboxRoot)
+{
 	if (CurrentCreationWorkflow)
 	{
 		CurrentCreationWorkflow->Cancel();
 	}
-	
+
 	Model->LoadSandbox(InSandboxRoot);
 }
 
diff --git a/Engine/Plugins/Developer/Sandbox/SandboxedEditing/Source/SandboxedEditing/Private/Features/Browser/ViewModels/SandboxControlsViewModel.h b/Engine/Plugins/Developer/Sandbox/SandboxedEditing/Source/SandboxedEditing/Private/Features/Browser/ViewModels/SandboxControlsViewModel.h
--- a/Engine/Plugins/Developer/Sandbox/SandboxedEditing/Source/SandboxedEditing/Private/Features/Browser/ViewModels/SandboxControlsViewModel.h
+++ b/Engine/Plugins/Developer/Sandbox/SandboxedEditing/Source/SandboxedEditing/Private/Features/Browser/ViewModels/SandboxControlsViewModel.h
@@ -6,11 +6,14 @@
 #include "Framework/Models/SandboxSystemModel.h"
 #include "RenameSandboxWorkflow.h"
 #include "SandboxCreationWorkflow.h"
+#include "Templates/SharedPointer.h"
 #include "Templates/UniquePtr.h"
 #include "Templates/UnrealTemplate.h"
 #include "Transfer/ExportWorkflow.h"
 #include "Transfer/ImportWorkflow.h"
 
+class SNotificationItem;
+
 namespace UE::SandboxedEditing
 {
 DECLARE_MULTICAST_DELEGATE_OneParam(FCreationWorkflowDelegate, FSandboxCreationWorkflow&);
@@ -20,7 +23,7 @@ DECLARE_MULTICAST_DELEGATE_OneParam(FImportWorkflowDelegate, FImportWorkflow&);
 
 /**
  * Handles actions that can be performed with FSandboxListItems, such as deleting, renaming, loading.
- * 
+ *
  * These actions are primarily present on the controls widget in the form of buttons.
  * Some of them can be performed in the list view, such as deleting or renaming.
  */
@@ -44,7 +47,7 @@ public:
 	bool CanLoadSandbox(const FString& InSandboxRoot, FText* OutReason = nullptr) const;
 	/** Loads the sandbox. */
 	void LoadSandbox(const FString& InSandboxRoot);
-	
+
 	/** @return Whether the sandbox can be deleted. */
 	bool CanDeleteSandbox(const FString& InSandboxRoot, FText* OutReason = nullptr) const;
 	/** Deletes the sandbox. */
@@ -90,11 +93,17 @@ public:
 	FImportWorkflowDelegate& OnImportWorkflowStarted() { return OnImportWorkflowStartedDelegate; }
 	/** @return Delegate that is invoked when the import workflow ends. */
 	FSimpleMulticastDelegate& OnImportWorkflowEnded() { return OnImportWorkflowEndedDelegate; }
-	
+
 private:
-	
+
+	/** Loads the sandbox without checking for in-memory changes. */
+	void LoadSandboxInternal(const FString& InSandboxRoot);
+
 	/** The underlying model that performs the mutations. */
 	const TSharedRef<FSandboxSystemModel> Model;
+
+	/** Active notification for in-memory changes warning. Lifetime is not bound to this class and only used to make sure we do not have multiple notification to the user. */
+	TWeakPtr<SNotificationItem> InMemoryChangesNotification;
 	
 	/** Set if the user is currently in the process of creating a sandbox. */
 	TUniquePtr<FSandboxCreationWorkflow> CurrentCreationWorkflow;
@@ -124,7 +133,7 @@ private:
 	FImportWorkflowDelegate OnImportWorkflowStartedDelegate;
 	/** Invoked when the import workflow ends. */
 	FSimpleMulticastDelegate OnImportWorkflowEndedDelegate;
-	
+
 	/** Called by CurrentCreationWorkflow when the workflow ends. */
 	void HandleCreationWorkflowEnded();
 	/** Called by CurrentRenameWorkflow when the workflow ends. */
diff --git a/Engine/Plugins/Developer/Sandbox/SandboxedEditing/Source/SandboxedEditing/Private/Framework/Notifications.cpp b/Engine/Plugins/Developer/Sandbox/SandboxedEditing/Source/SandboxedEditing/Private/Framework/Notifications.cpp
--- a/Engine/Plugins/Developer/Sandbox/SandboxedEditing/Source/SandboxedEditing/Private/Framework/Notifications.cpp
+++ b/Engine/Plugins/Developer/Sandbox/SandboxedEditing/Source/SandboxedEditing/Private/Framework/Notifications.cpp
@@ -69,16 +69,75 @@ void ShowLeftSandbox(const FString& InSandboxName)
 	{
 		return;
 	}
-	
+
 	FSlateNotificationManager& NotificationManager = FSlateNotificationManager::Get();
 	FNotificationInfo Args(
 		FText::Format(LOCTEXT("LeftFmt", "Left sandbox {0}"), FText::FromString(InSandboxName))
 		);
 	Args.bFireAndForget = true;
 	Args.ExpireDuration = 2.f;
-	
+
 	NotificationManager.AddNotification(Args);
 }
+
+TSharedPtr<SNotificationItem> ShowInMemoryChangesWarning(const FString& InSandboxName, TFunction<void()> OnDiscard, TFunction<void()> OnCancel)
+{
+	if (!FSlateApplication::IsInitialized())
+	{
+		return nullptr;
+	}
+
+	FSlateNotificationManager& NotificationManager = FSlateNotificationManager::Get();
+
+	// Create a shared pointer to hold the weak notification so lambdas can capture it by value
+	TSharedPtr<TWeakPtr<SNotificationItem>> WeakNotificationPtr = MakeShared<TWeakPtr<SNotificationItem>>();
+
+	FNotificationInfo Info(
+		FText::Format(LOCTEXT("LoadingSandbox", "Loading {0}"), FText::FromString(InSandboxName))
+	);
+	Info.SubText = LOCTEXT("InMemoryChanges.Body", "The project has in-memory changes. Continuing will discard changes.");
+	Info.bFireAndForget = false;
+	Info.ExpireDuration = 0.0f;
+	Info.ButtonDetails =
+	{
+		FNotificationButtonInfo(
+			LOCTEXT("Discard", "Discard"),
+			LOCTEXT("Discard.ToolTip", "Discard in-memory changes and load sandbox"),
+			FSimpleDelegate::CreateLambda([OnDiscard, WeakNotificationPtr]()
+			{
+				if (TSharedPtr<SNotificationItem> PinnedNotification = WeakNotificationPtr->Pin())
+				{
+					PinnedNotification->ExpireAndFadeout();
+				}
+				if (OnDiscard)
+				{
+					OnDiscard();
+				}
+			}),
+			SNotificationItem::CS_None
+		),
+		FNotificationButtonInfo(
+			LOCTEXT("Cancel", "Cancel"),
+			LOCTEXT("Cancel.ToolTip", "Cancel loading sandbox"),
+			FSimpleDelegate::CreateLambda([OnCancel, WeakNotificationPtr]()
+			{
+				if (TSharedPtr<SNotificationItem> PinnedNotification = WeakNotificationPtr->Pin())
+				{
+					PinnedNotification->ExpireAndFadeout();
+				}
+				if (OnCancel)
+				{
+					OnCancel();
+				}
+			}),
+			SNotificationItem::CS_None
+		)
+	};
+
+	TSharedPtr<SNotificationItem> Notification = NotificationManager.AddNotification(Info);
+	*WeakNotificationPtr = Notification;
+	return Notification;
+}
 }
 
-#undef LOCTEXT_NAMESPACE
\ No newline at end of file
+#undef LOCTEXT_NAMESPACE
diff --git a/Engine/Plugins/Developer/Sandbox/SandboxedEditing/Source/SandboxedEditing/Private/Framework/Notifications.h b/Engine/Plugins/Developer/Sandbox/SandboxedEditing/Source/SandboxedEditing/Private/Framework/Notifications.h
--- a/Engine/Plugins/Developer/Sandbox/SandboxedEditing/Source/SandboxedEditing/Private/Framework/Notifications.h
+++ b/Engine/Plugins/Developer/Sandbox/SandboxedEditing/Source/SandboxedEditing/Private/Framework/Notifications.h
@@ -2,13 +2,17 @@
 
 #pragma once
 
+#include "Templates/Function.h"
+#include "Templates/SharedPointer.h"
+
 class FString;
+class SNotificationItem;
 
 namespace UE::SandboxedEditing
 {
 /** Shows a notification that the user joined a sandbox. */
 void ShowCreatedSandbox(const FString& InSandboxName);
-	
+
 /** Shows the sandbox was loaded successfully. */
 void ShowLoadedSandbox(const FString& InSandboxName);
 /** Show that the sandbox could not be loaded. */
@@ -16,4 +20,14 @@ void ShowFailedToLoadSandbox(const FString& InSandboxName);
 
 /** Shows notification that the user left a sandbox. */
 void ShowLeftSandbox(const FString& InSandboxName);
+
+/**
+ * Shows a notification warning the user about in-memory changes before loading a sandbox.
+ * @param InSandboxName The name of the sandbox being loaded
+ * @param OnDiscard Callback invoked when user chooses to discard changes
+ * @param OnCancel Callback invoked when user chooses to cancel
+ * @return The notification item.
+ */
+TSharedPtr<SNotificationItem> ShowInMemoryChangesWarning(const FString& InSandboxName, TFunction<void()> OnDiscard, TFunction<void()> OnCancel);
 }
+
