Add sortable columns to the activity list in the sandbox.

#rnx

#Synced-CL 52444341

diff --git a/Engine/Plugins/Developer/Sandbox/SandboxedEditing/Source/SandboxedEditing/Private/Features/Browser/ViewModels/FileState/IFileStateColumnBehavior.h b/Engine/Plugins/Developer/Sandbox/SandboxedEditing/Source/SandboxedEditing/Private/Features/Browser/ViewModels/FileState/IFileStateColumnBehavior.h
--- a/Engine/Plugins/Developer/Sandbox/SandboxedEditing/Source/SandboxedEditing/Private/Features/Browser/ViewModels/FileState/IFileStateColumnBehavior.h
+++ b/Engine/Plugins/Developer/Sandbox/SandboxedEditing/Source/SandboxedEditing/Private/Features/Browser/ViewModels/FileState/IFileStateColumnBehavior.h
@@ -3,6 +3,7 @@
 #pragma once
 
 #include "Features/Browser/ViewModels/Columns/ISearchableColumnBehavior.h"
+#include "Features/Browser/ViewModels/Columns/ISortableColumnBehavior.h"
 #include "Templates/SharedPointerFwd.h"
 
 namespace UE::SandboxedEditing
@@ -11,8 +12,10 @@ struct FFileStateItem;
 
 /**
  * View-model logic that is implemented by columns.
- * TODO UE-358389: Add interface for sorting
+ * Supports both searching and sorting.
  */
-class IFileStateColumnBehavior : public ISearchableColumnBehavior<TSharedPtr<FFileStateItem>>
+class IFileStateColumnBehavior
+	: public ISearchableColumnBehavior<TSharedPtr<FFileStateItem>>
+	, public ISortableColumnBehavior<TSharedPtr<FFileStateItem>>
 {};
 }
diff --git a/Engine/Plugins/Developer/Sandbox/SandboxedEditing/Source/SandboxedEditing/Private/Features/Browser/ViewModels/FileState/IFileStateViewModel.h b/Engine/Plugins/Developer/Sandbox/SandboxedEditing/Source/SandboxedEditing/Private/Features/Browser/ViewModels/FileState/IFileStateViewModel.h
--- a/Engine/Plugins/Developer/Sandbox/SandboxedEditing/Source/SandboxedEditing/Private/Features/Browser/ViewModels/FileState/IFileStateViewModel.h
+++ b/Engine/Plugins/Developer/Sandbox/SandboxedEditing/Source/SandboxedEditing/Private/Features/Browser/ViewModels/FileState/IFileStateViewModel.h
@@ -1,16 +1,17 @@
-﻿// Copyright Epic Games, Inc. All Rights Reserved.
+// Copyright Epic Games, Inc. All Rights Reserved.
 
 #pragma once
 
 #include "Delegates/Delegate.h"
+#include "Widgets/Views/SHeaderRow.h"
 
 namespace UE::SandboxedEditing
 {
 struct FFileStateItem;
 
-/** 
- * View model for UI that shows the files added / removed / edited in the sandbox. 
- * 
+/**
+ * View model for UI that shows the files added / removed / edited in the sandbox.
+ *
  * The implementation can vary based on where the UI is displayed, e.g.
  * - Persist files (part of workflow for leaving sandbox)
  * - Sandbox changed files (part of UI for displaying current sandbox details)
@@ -18,16 +19,43 @@ struct FFileStateItem;
 class IFileStateViewModel
 {
 public:
-	
+
 	/** @return The items to display */
 	virtual const TArray<TSharedPtr<FFileStateItem>>& GetItems() const = 0;
-	
+
 	/** @return The columns that should be displayed. */
 	virtual TArray<FName> GetDisplayedColumns() const = 0;
-	
+
 	/** Invoked when the items change. */
 	virtual FSimpleMulticastDelegate& OnItemsChanged() = 0;
-	
+
+	/**
+	 * Returns the current column sort mode (ascending or descending) if the ColumnId parameter matches the current
+	 * column to be sorted by, otherwise returns EColumnSortMode::None.
+	 *
+	 * @param ColumnId Column ID to query sort mode for.
+	 * @return The sort mode for the column, or EColumnSortMode::None if it is not known.
+	 */
+	virtual EColumnSortMode::Type GetColumnSortMode(const FName ColumnId) const = 0;
+
+	/**
+	 * Returns the current column sort priority (primary or secondary) if the ColumnId parameter matches a sorted column,
+	 * otherwise returns EColumnSortPriority::Max.
+	 *
+	 * @param ColumnId Column ID to query sort priority for.
+	 * @return The sort priority for the column, or EColumnSortPriority::Max if it is not sorted.
+	 */
+	virtual EColumnSortPriority::Type GetColumnSortPriority(const FName ColumnId) const = 0;
+
+	/**
+	 * Callback for SHeaderRow::Column::OnSort, called when the column to sort by is changed.
+	 *
+	 * @param SortPriority The sort priority (Primary or Secondary)
+	 * @param ColumnId The column to sort by
+	 * @param InSortMode The sort mode (ascending or descending)
+	 */
+	virtual void OnColumnSortModeChanged(const EColumnSortPriority::Type SortPriority, const FName& ColumnId, const EColumnSortMode::Type InSortMode) = 0;
+
 	virtual ~IFileStateViewModel() = default;
 };
 }
diff --git a/Engine/Plugins/Developer/Sandbox/SandboxedEditing/Source/SandboxedEditing/Private/Features/Browser/ViewModels/FileState/Models/StaticFileStateViewModel.cpp b/Engine/Plugins/Developer/Sandbox/SandboxedEditing/Source/SandboxedEditing/Private/Features/Browser/ViewModels/FileState/Models/StaticFileStateViewModel.cpp
--- a/Engine/Plugins/Developer/Sandbox/SandboxedEditing/Source/SandboxedEditing/Private/Features/Browser/ViewModels/FileState/Models/StaticFileStateViewModel.cpp
+++ b/Engine/Plugins/Developer/Sandbox/SandboxedEditing/Source/SandboxedEditing/Private/Features/Browser/ViewModels/FileState/Models/StaticFileStateViewModel.cpp
@@ -3,6 +3,7 @@
 #include "StaticFileStateViewModel.h"
 
 #include "Features/Browser/ViewModels/FileState/FilterFileStateViewModel.h"
+#include "Features/Browser/ViewModels/FileState/IFileStateColumnBehavior.h"
 
 namespace UE::SandboxedEditing
 {
@@ -50,9 +51,9 @@ void FStaticFileStateViewModel::FilterItems()
 		FilteredItems = AllUnfilteredItems;
 		return;
 	}
-	
+
 	FilteredItems.Empty(AllUnfilteredItems.Num());
-	
+
 	for (const TSharedPtr<FFileStateItem>& Item : AllUnfilteredItems)
 	{
 		if (ensure(Item) && FilterViewModel->PassesFilter(Item.ToSharedRef()))
@@ -61,4 +62,113 @@ void FStaticFileStateViewModel::FilterItems()
 		}
 	}
 }
+
+EColumnSortMode::Type FStaticFileStateViewModel::GetColumnSortMode(const FName ColumnId) const
+{
+	if (ColumnId == PrimarySortedColumn)
+	{
+		return PrimarySortMode;
+	}
+	if (ColumnId == SecondarySortedColumn)
+	{
+		return SecondarySortMode;
+	}
+	return EColumnSortMode::None;
+}
+
+EColumnSortPriority::Type FStaticFileStateViewModel::GetColumnSortPriority(const FName ColumnId) const
+{
+	if (ColumnId == PrimarySortedColumn)
+	{
+		return EColumnSortPriority::Primary;
+	}
+	if (ColumnId == SecondarySortedColumn)
+	{
+		return EColumnSortPriority::Secondary;
+	}
+	return EColumnSortPriority::Max;
+}
+
+void FStaticFileStateViewModel::OnColumnSortModeChanged(
+	const EColumnSortPriority::Type SortPriority,
+	const FName& ColumnId,
+	const EColumnSortMode::Type InSortMode)
+{
+	if (SortPriority == EColumnSortPriority::Primary)
+	{
+		PrimarySortedColumn = ColumnId;
+		PrimarySortMode = InSortMode;
+
+		// Cannot be primary and secondary at the same time
+		if (ColumnId == SecondarySortedColumn)
+		{
+			SecondarySortedColumn = FName();
+			SecondarySortMode = EColumnSortMode::None;
+		}
+	}
+	else if (SortPriority == EColumnSortPriority::Secondary)
+	{
+		SecondarySortedColumn = ColumnId;
+		SecondarySortMode = InSortMode;
+	}
+
+	RequestSort();
+}
+
+void FStaticFileStateViewModel::RequestSort()
+{
+	SortItems();
+	OnItemsChangedDelegate.Broadcast();
+}
+
+void FStaticFileStateViewModel::SortItems()
+{
+	if (PrimarySortedColumn.IsNone())
+	{
+		return; // No sorting set
+	}
+
+	const TSharedRef<IFileStateColumnBehavior>* PrimaryColumnBehavior = Columns.Find(PrimarySortedColumn);
+	if (!PrimaryColumnBehavior || !PrimaryColumnBehavior->Get().CanSort())
+	{
+		return;
+	}
+
+	const TSharedRef<IFileStateColumnBehavior>* SecondaryColumnBehavior = nullptr;
+	if (!SecondarySortedColumn.IsNone())
+	{
+		SecondaryColumnBehavior = Columns.Find(SecondarySortedColumn);
+		if (SecondaryColumnBehavior && !SecondaryColumnBehavior->Get().CanSort())
+		{
+			SecondaryColumnBehavior = nullptr;
+		}
+	}
+
+	FilteredItems.Sort([&](const TSharedPtr<FFileStateItem>& Lhs, const TSharedPtr<FFileStateItem>& Rhs)
+	{
+		if (!Lhs || !Rhs)
+		{
+			return false;
+		}
+
+		// Try primary sort first
+		if (PrimaryColumnBehavior->Get().SortBy(Lhs, Rhs, PrimarySortMode))
+		{
+			return true; // Lhs must be before Rhs based on primary sort
+		}
+		// Check if they're equal on primary column by testing reverse
+		if (PrimaryColumnBehavior->Get().SortBy(Rhs, Lhs, PrimarySortMode))
+		{
+			return false; // Rhs must be before Lhs based on primary sort
+		}
+
+		// Equal on primary column - use secondary if available
+		if (SecondaryColumnBehavior)
+		{
+			return SecondaryColumnBehavior->Get().SortBy(Lhs, Rhs, SecondarySortMode);
+		}
+
+		return false;
+	});
+}
 }
diff --git a/Engine/Plugins/Developer/Sandbox/SandboxedEditing/Source/SandboxedEditing/Private/Features/Browser/ViewModels/FileState/Models/StaticFileStateViewModel.h b/Engine/Plugins/Developer/Sandbox/SandboxedEditing/Source/SandboxedEditing/Private/Features/Browser/ViewModels/FileState/Models/StaticFileStateViewModel.h
--- a/Engine/Plugins/Developer/Sandbox/SandboxedEditing/Source/SandboxedEditing/Private/Features/Browser/ViewModels/FileState/Models/StaticFileStateViewModel.h
+++ b/Engine/Plugins/Developer/Sandbox/SandboxedEditing/Source/SandboxedEditing/Private/Features/Browser/ViewModels/FileState/Models/StaticFileStateViewModel.h
@@ -32,6 +32,9 @@ public:
 	virtual const TArray<TSharedPtr<FFileStateItem>>& GetItems() const override { return FilteredItems; }
 	virtual TArray<FName> GetDisplayedColumns() const override;
 	virtual FSimpleMulticastDelegate& OnItemsChanged() override { return OnItemsChangedDelegate; }
+	virtual EColumnSortMode::Type GetColumnSortMode(const FName ColumnId) const override;
+	virtual EColumnSortPriority::Type GetColumnSortPriority(const FName ColumnId) const override;
+	virtual void OnColumnSortModeChanged(const EColumnSortPriority::Type SortPriority, const FName& ColumnId, const EColumnSortMode::Type InSortMode) override;
 	//~ End IFileStateViewModel Interface
 
 protected:
@@ -49,18 +52,30 @@ protected:
 	void NotifyItemsChanged();
 	
 private:
-	
-	/**
-	 * Dummy to fulfill the interface requirements. Can be used by subclasses. 
-	 * TODO when sorting column is implemented, this can be invoked.
-	 */
+
+	/** Invoked when items or sorting changes. */
 	FSimpleMulticastDelegate OnItemsChangedDelegate;
-	
+
 	/** A filtered version of AllUnfilteredItems. */
 	TArray<TSharedPtr<FFileStateItem>> FilteredItems;
-	
+
+	/** Sorting state */
+	FName PrimarySortedColumn;
+	EColumnSortMode::Type PrimarySortMode = EColumnSortMode::None;
+	FName SecondarySortedColumn;
+	EColumnSortMode::Type SecondarySortMode = EColumnSortMode::None;
+
 	/** Refilters the items. */
 	void FilterItems();
+
+	/**
+	 * Requests that the items be sorted according to the current sort column and mode.
+	 * Refreshes the list after sorting.
+	 */
+	void RequestSort();
+
+	/** Sorts the items according to the current sort column and mode. */
+	void SortItems();
 };
 }
 
diff --git a/Engine/Plugins/Developer/Sandbox/SandboxedEditing/Source/SandboxedEditing/Private/Features/Browser/Widgets/Shared/FileState/Columns/ActionFileStateColumn.cpp b/Engine/Plugins/Developer/Sandbox/SandboxedEditing/Source/SandboxedEditing/Private/Features/Browser/Widgets/Shared/FileState/Columns/ActionFileStateColumn.cpp
--- a/Engine/Plugins/Developer/Sandbox/SandboxedEditing/Source/SandboxedEditing/Private/Features/Browser/Widgets/Shared/FileState/Columns/ActionFileStateColumn.cpp
+++ b/Engine/Plugins/Developer/Sandbox/SandboxedEditing/Source/SandboxedEditing/Private/Features/Browser/Widgets/Shared/FileState/Columns/ActionFileStateColumn.cpp
@@ -60,6 +60,14 @@ void FActionFileStateColumn::PopulateSearchTerms(const TSharedPtr<FFileStateItem
 	}
 }
 
+bool FActionFileStateColumn::SortBy(const TSharedPtr<FFileStateItem>& Lhs, const TSharedPtr<FFileStateItem>& Rhs, EColumnSortMode::Type SortMode) const
+{
+	// Sort by action type (enum value order)
+	return SortMode == EColumnSortMode::Ascending
+		? static_cast<uint8>(Lhs->Action) < static_cast<uint8>(Rhs->Action)
+		: static_cast<uint8>(Lhs->Action) > static_cast<uint8>(Rhs->Action);
+}
+
 SHeaderRow::FColumn::FArguments FActionFileStateColumn::MakeColumnArguments()
 {
 	return SHeaderRow::FColumn::FArguments()
diff --git a/Engine/Plugins/Developer/Sandbox/SandboxedEditing/Source/SandboxedEditing/Private/Features/Browser/Widgets/Shared/FileState/Columns/ActionFileStateColumn.h b/Engine/Plugins/Developer/Sandbox/SandboxedEditing/Source/SandboxedEditing/Private/Features/Browser/Widgets/Shared/FileState/Columns/ActionFileStateColumn.h
--- a/Engine/Plugins/Developer/Sandbox/SandboxedEditing/Source/SandboxedEditing/Private/Features/Browser/Widgets/Shared/FileState/Columns/ActionFileStateColumn.h
+++ b/Engine/Plugins/Developer/Sandbox/SandboxedEditing/Source/SandboxedEditing/Private/Features/Browser/Widgets/Shared/FileState/Columns/ActionFileStateColumn.h
@@ -13,6 +13,8 @@ public:
 	
 	//~ Begin IFileStateColumnBehavior Interface
 	virtual void PopulateSearchTerms(const TSharedPtr<FFileStateItem>& InRowData, TArray<FString>& OutSearchTerms) const override;
+	virtual bool CanSort() const override { return true; }
+	virtual bool SortBy(const TSharedPtr<FFileStateItem>& Lhs, const TSharedPtr<FFileStateItem>& Rhs, EColumnSortMode::Type SortMode) const override;
 	//~ End IFileStateColumnBehavior Interface
 	
 	//~ Begin IFileStateWidgetFactory Interface
diff --git a/Engine/Plugins/Developer/Sandbox/SandboxedEditing/Source/SandboxedEditing/Private/Features/Browser/Widgets/Shared/FileState/Columns/PathFileStateColumn.cpp b/Engine/Plugins/Developer/Sandbox/SandboxedEditing/Source/SandboxedEditing/Private/Features/Browser/Widgets/Shared/FileState/Columns/PathFileStateColumn.cpp
--- a/Engine/Plugins/Developer/Sandbox/SandboxedEditing/Source/SandboxedEditing/Private/Features/Browser/Widgets/Shared/FileState/Columns/PathFileStateColumn.cpp
+++ b/Engine/Plugins/Developer/Sandbox/SandboxedEditing/Source/SandboxedEditing/Private/Features/Browser/Widgets/Shared/FileState/Columns/PathFileStateColumn.cpp
@@ -45,6 +45,13 @@ void FPathFileStateColumn::PopulateSearchTerms(const TSharedPtr<FFileStateItem>&
 		);
 }
 
+bool FPathFileStateColumn::SortBy(const TSharedPtr<FFileStateItem>& Lhs, const TSharedPtr<FFileStateItem>& Rhs, EColumnSortMode::Type SortMode) const
+{
+	const FString LhsPath = PathFileColumnDetail::MakeProjectRelative(Lhs->NonSandboxFile);
+	const FString RhsPath = PathFileColumnDetail::MakeProjectRelative(Rhs->NonSandboxFile);
+	return SortMode == EColumnSortMode::Ascending ? LhsPath < RhsPath : LhsPath > RhsPath;
+}
+
 SHeaderRow::FColumn::FArguments FPathFileStateColumn::MakeColumnArguments()
 {
 	return SHeaderRow::FColumn::FArguments()
diff --git a/Engine/Plugins/Developer/Sandbox/SandboxedEditing/Source/SandboxedEditing/Private/Features/Browser/Widgets/Shared/FileState/Columns/PathFileStateColumn.h b/Engine/Plugins/Developer/Sandbox/SandboxedEditing/Source/SandboxedEditing/Private/Features/Browser/Widgets/Shared/FileState/Columns/PathFileStateColumn.h
--- a/Engine/Plugins/Developer/Sandbox/SandboxedEditing/Source/SandboxedEditing/Private/Features/Browser/Widgets/Shared/FileState/Columns/PathFileStateColumn.h
+++ b/Engine/Plugins/Developer/Sandbox/SandboxedEditing/Source/SandboxedEditing/Private/Features/Browser/Widgets/Shared/FileState/Columns/PathFileStateColumn.h
@@ -14,6 +14,8 @@ public:
 	
 	//~ Begin IFileStateColumnBehavior Interface
 	virtual void PopulateSearchTerms(const TSharedPtr<FFileStateItem>& InRowData, TArray<FString>& OutSearchTerms) const override;
+	virtual bool CanSort() const override { return true; }
+	virtual bool SortBy(const TSharedPtr<FFileStateItem>& Lhs, const TSharedPtr<FFileStateItem>& Rhs, EColumnSortMode::Type SortMode) const override;
 	//~ End IFileStateColumnBehavior Interface
 	
 	//~ Begin IFileStateWidgetFactory Interface
diff --git a/Engine/Plugins/Developer/Sandbox/SandboxedEditing/Source/SandboxedEditing/Private/Features/Browser/Widgets/Shared/FileState/Columns/PersistFileStateColumn.h b/Engine/Plugins/Developer/Sandbox/SandboxedEditing/Source/SandboxedEditing/Private/Features/Browser/Widgets/Shared/FileState/Columns/PersistFileStateColumn.h
--- a/Engine/Plugins/Developer/Sandbox/SandboxedEditing/Source/SandboxedEditing/Private/Features/Browser/Widgets/Shared/FileState/Columns/PersistFileStateColumn.h
+++ b/Engine/Plugins/Developer/Sandbox/SandboxedEditing/Source/SandboxedEditing/Private/Features/Browser/Widgets/Shared/FileState/Columns/PersistFileStateColumn.h
@@ -17,6 +17,8 @@ public:
 
 	//~ Begin IFileStateColumnBehavior Interface
 	virtual void PopulateSearchTerms(const TSharedPtr<FFileStateItem>& InRowData, TArray<FString>& OutSearchTerms) const override {}
+	virtual bool CanSort() const override { return false; }
+	virtual bool SortBy(const TSharedPtr<FFileStateItem>& Lhs, const TSharedPtr<FFileStateItem>& Rhs, EColumnSortMode::Type SortMode) const override { return false; }
 	//~ End IFileStateColumnBehavior Interface
 
 	//~ Begin IFileStateWidgetFactory Interface
diff --git a/Engine/Plugins/Developer/Sandbox/SandboxedEditing/Source/SandboxedEditing/Private/Features/Browser/Widgets/Shared/FileState/Columns/TimestampFileStateColumn.cpp b/Engine/Plugins/Developer/Sandbox/SandboxedEditing/Source/SandboxedEditing/Private/Features/Browser/Widgets/Shared/FileState/Columns/TimestampFileStateColumn.cpp
--- a/Engine/Plugins/Developer/Sandbox/SandboxedEditing/Source/SandboxedEditing/Private/Features/Browser/Widgets/Shared/FileState/Columns/TimestampFileStateColumn.cpp
+++ b/Engine/Plugins/Developer/Sandbox/SandboxedEditing/Source/SandboxedEditing/Private/Features/Browser/Widgets/Shared/FileState/Columns/TimestampFileStateColumn.cpp
@@ -67,6 +67,14 @@ void FTimestampFileStateColumn::PopulateSearchTerms(const TSharedPtr<FFileStateI
 		);
 }
 
+bool FTimestampFileStateColumn::SortBy(const TSharedPtr<FFileStateItem>& Lhs, const TSharedPtr<FFileStateItem>& Rhs, EColumnSortMode::Type SortMode) const
+{
+	// Sort by timestamp (most recent first in descending, oldest first in ascending)
+	return SortMode == EColumnSortMode::Ascending
+		? Lhs->Timestamp < Rhs->Timestamp
+		: Lhs->Timestamp > Rhs->Timestamp;
+}
+
 SHeaderRow::FColumn::FArguments FTimestampFileStateColumn::MakeColumnArguments()
 {
 	return SHeaderRow::FColumn::FArguments()
diff --git a/Engine/Plugins/Developer/Sandbox/SandboxedEditing/Source/SandboxedEditing/Private/Features/Browser/Widgets/Shared/FileState/Columns/TimestampFileStateColumn.h b/Engine/Plugins/Developer/Sandbox/SandboxedEditing/Source/SandboxedEditing/Private/Features/Browser/Widgets/Shared/FileState/Columns/TimestampFileStateColumn.h
--- a/Engine/Plugins/Developer/Sandbox/SandboxedEditing/Source/SandboxedEditing/Private/Features/Browser/Widgets/Shared/FileState/Columns/TimestampFileStateColumn.h
+++ b/Engine/Plugins/Developer/Sandbox/SandboxedEditing/Source/SandboxedEditing/Private/Features/Browser/Widgets/Shared/FileState/Columns/TimestampFileStateColumn.h
@@ -14,6 +14,8 @@ public:
 	
 	//~ Begin IFileStateColumnBehavior Interface
 	virtual void PopulateSearchTerms(const TSharedPtr<FFileStateItem>& InRowData, TArray<FString>& OutSearchTerms) const override;
+	virtual bool CanSort() const override { return true; }
+	virtual bool SortBy(const TSharedPtr<FFileStateItem>& Lhs, const TSharedPtr<FFileStateItem>& Rhs, EColumnSortMode::Type SortMode) const override;
 	//~ End IFileStateColumnBehavior Interface
 
 	//~ Begin IFileStateWidgetFactory Interface
diff --git a/Engine/Plugins/Developer/Sandbox/SandboxedEditing/Source/SandboxedEditing/Private/Features/Browser/Widgets/Shared/FileState/SFileStateListView.cpp b/Engine/Plugins/Developer/Sandbox/SandboxedEditing/Source/SandboxedEditing/Private/Features/Browser/Widgets/Shared/FileState/SFileStateListView.cpp
--- a/Engine/Plugins/Developer/Sandbox/SandboxedEditing/Source/SandboxedEditing/Private/Features/Browser/Widgets/Shared/FileState/SFileStateListView.cpp
+++ b/Engine/Plugins/Developer/Sandbox/SandboxedEditing/Source/SandboxedEditing/Private/Features/Browser/Widgets/Shared/FileState/SFileStateListView.cpp
@@ -62,15 +62,20 @@ TArray<TSharedPtr<FFileStateItem>> SFileStateListView::GetSelectedFiles() const
 TSharedRef<SHeaderRow> SFileStateListView::MakeHeaderRow()
 {
 	const TSharedRef<SHeaderRow> HeaderRow = SNew(SHeaderRow);
-		
+
 	for (FName ColumnId : ViewModel->GetDisplayedColumns())
 	{
 		if (const TSharedRef<IFileStateColumnWidgetFactory>* Factory = ColumnFactories.Find(ColumnId))
 		{
-			HeaderRow->AddColumn(Factory->Get().MakeColumnArguments());
+			SHeaderRow::FColumn::FArguments ColumnArgs = Factory->Get().MakeColumnArguments();
+			ColumnArgs
+				.SortMode(ViewModel.ToSharedRef(), &IFileStateViewModel::GetColumnSortMode, ColumnId)
+				.SortPriority(ViewModel.ToSharedRef(), &IFileStateViewModel::GetColumnSortPriority, ColumnId)
+				.OnSort(ViewModel.ToSharedRef(), &IFileStateViewModel::OnColumnSortModeChanged);
+			HeaderRow->AddColumn(ColumnArgs);
 		}
 	}
-	
+
 	return HeaderRow;
 }
 
