Sandboxed Editing: Improve context-menu labels wording.

#Synced-CL 52378315

diff --git a/Engine/Plugins/Developer/Sandbox/SandboxedEditing/Source/SandboxedEditing/Private/Features/Browser/Commands/FileStateActions/FileStateCommands.cpp b/Engine/Plugins/Developer/Sandbox/SandboxedEditing/Source/SandboxedEditing/Private/Features/Browser/Commands/FileStateActions/FileStateCommands.cpp
--- a/Engine/Plugins/Developer/Sandbox/SandboxedEditing/Source/SandboxedEditing/Private/Features/Browser/Commands/FileStateActions/FileStateCommands.cpp
+++ b/Engine/Plugins/Developer/Sandbox/SandboxedEditing/Source/SandboxedEditing/Private/Features/Browser/Commands/FileStateActions/FileStateCommands.cpp
@@ -20,8 +20,8 @@ FFileStateCommands::FFileStateCommands()
 void FFileStateCommands::RegisterCommands()
 {
 	UI_COMMAND(BrowseToAsset, "Browse to asset", "Shows this asset in the content browser.", EUserInterfaceActionType::Button, FInputChord())
-	UI_COMMAND(ShowRootInExplorer, "Show sandbox folder", "Shows sandbox root directory in the system file manager.", EUserInterfaceActionType::Button, FInputChord())
-	UI_COMMAND(ShowFileInExplorer, "Show in file manager", "Shows the file's sandboxed counterpart in the system file manager.", EUserInterfaceActionType::Button, FInputChord())
+	UI_COMMAND(ShowRootInExplorer, "Show root sandbox folder", "Shows sandbox root directory in the system file manager.", EUserInterfaceActionType::Button, FInputChord())
+	UI_COMMAND(ShowFileInExplorer, "Show sandbox file", "Shows the file's sandboxed counterpart in the system file manager.", EUserInterfaceActionType::Button, FInputChord())
 	
 	UI_COMMAND(PersistSelected, "Persist", "Open the persist dialogue with the currently selected files already pre-selected in the dialogue.", EUserInterfaceActionType::Button, FInputChord())
 }
