Add Content Browser badge feature for sandboxed files

Implements visual badges on Content Browser asset thumbnails to indicate
files that have been modified, added, or deleted in the active sandbox.

- Added FContentBrowserBadgeFeature class to manage badge generation
- Badges show appropriate icons with color coding:
* Added files: Blue plus icon
* Modified files: Green edit icon
* Deleted files: Red delete icon
- Tooltips display sandbox name and change type
- Badges automatically register/unregister on sandbox load/leave
- Path normalization ensures proper comparison between package paths and sandbox changes

#rnx

#Synced-CL 52627824

diff --git a/Engine/Plugins/Developer/Sandbox/SandboxedEditing/Source/SandboxedEditing/Private/Features/ContentBrowserBadge/ContentBrowserBadgeFeature.cpp b/Engine/Plugins/Developer/Sandbox/SandboxedEditing/Source/SandboxedEditing/Private/Features/ContentBrowserBadge/ContentBrowserBadgeFeature.cpp
new file mode 100644
--- /dev/null
+++ b/Engine/Plugins/Developer/Sandbox/SandboxedEditing/Source/SandboxedEditing/Private/Features/ContentBrowserBadge/ContentBrowserBadgeFeature.cpp
@@ -0,0 +1,216 @@
+// Copyright Epic Games, Inc. All Rights Reserved.
+
+#include "ContentBrowserBadgeFeature.h"
+#include "Framework/Models/SandboxSystemModel.h"
+#include "ContentBrowserModule.h"
+#include "AssetRegistry/AssetData.h"
+#include "Misc/PackageName.h"
+#include "Misc/Paths.h"
+#include "Styling/AppStyle.h"
+#include "Types/EBreakBehavior.h"
+#include "Types/SandboxedFileChangeInfo.h"
+#include "Widgets/Images/SImage.h"
+#include "Widgets/Text/STextBlock.h"
+
+namespace UE::SandboxedEditing
+{
+
+FContentBrowserBadgeFeature::FContentBrowserBadgeFeature(TSharedRef<FSandboxSystemModel> InSandboxModel)
+	: SandboxModel(InSandboxModel)
+{
+	// Subscribe to sandbox lifecycle events
+	OnLoadSandboxHandle = SandboxModel->OnLoadSandbox().AddRaw(this, &FContentBrowserBadgeFeature::OnSandboxLoaded);
+	OnLeaveSandboxHandle = SandboxModel->OnLeaveSandbox().AddRaw(this, &FContentBrowserBadgeFeature::OnSandboxLeft);
+
+	// If a sandbox is already active, register immediately
+	if (SandboxModel->HasActiveSandbox())
+	{
+		RegisterBadgeGenerators();
+	}
+}
+
+FContentBrowserBadgeFeature::~FContentBrowserBadgeFeature()
+{
+	// Unregister badge generators if still registered
+	UnregisterBadgeGenerators();
+
+	// Unsubscribe from events
+	if (OnLoadSandboxHandle.IsValid())
+	{
+		SandboxModel->OnLoadSandbox().Remove(OnLoadSandboxHandle);
+	}
+	if (OnLeaveSandboxHandle.IsValid())
+	{
+		SandboxModel->OnLeaveSandbox().Remove(OnLeaveSandboxHandle);
+	}
+}
+
+void FContentBrowserBadgeFeature::RegisterBadgeGenerators()
+{
+	if (BadgeGeneratorHandle.IsValid())
+	{
+		// Already registered
+		return;
+	}
+
+	FContentBrowserModule& ContentBrowserModule = FModuleManager::LoadModuleChecked<FContentBrowserModule>("ContentBrowser");
+
+	// Create the badge generator with both icon and tooltip delegates
+	FAssetViewExtraStateGenerator BadgeGenerator(
+		FOnGenerateAssetViewExtraStateIndicators::CreateRaw(this, &FContentBrowserBadgeFeature::GenerateBadgeIcon),
+		FOnGenerateAssetViewExtraStateIndicators::CreateRaw(this, &FContentBrowserBadgeFeature::GenerateBadgeTooltip)
+	);
+
+	BadgeGeneratorHandle = ContentBrowserModule.AddAssetViewExtraStateGenerator(BadgeGenerator);
+}
+
+void FContentBrowserBadgeFeature::UnregisterBadgeGenerators()
+{
+	if (!BadgeGeneratorHandle.IsValid())
+	{
+		return;
+	}
+
+	FContentBrowserModule& ContentBrowserModule = FModuleManager::LoadModuleChecked<FContentBrowserModule>("ContentBrowser");
+	ContentBrowserModule.RemoveAssetViewExtraStateGenerator(BadgeGeneratorHandle);
+	BadgeGeneratorHandle.Reset();
+}
+
+void FContentBrowserBadgeFeature::OnSandboxLoaded()
+{
+	RegisterBadgeGenerators();
+}
+
+void FContentBrowserBadgeFeature::OnSandboxLeft()
+{
+	UnregisterBadgeGenerators();
+}
+
+TSharedRef<SWidget> FContentBrowserBadgeFeature::GenerateBadgeIcon(const FAssetData& AssetData)
+{
+	const ESandboxChangeType ChangeType = GetAssetChangeType(AssetData);
+
+	if (ChangeType == ESandboxChangeType::None)
+	{
+		// No badge for unchanged assets
+		return SNullWidget::NullWidget;
+	}
+
+	// Choose icon and color based on change type
+	const FSlateBrush* IconBrush = nullptr;
+	FLinearColor BadgeColor = FLinearColor::White;
+
+	switch (ChangeType)
+	{
+	case ESandboxChangeType::Added:
+		IconBrush = FAppStyle::GetBrush("Icons.Plus");
+		BadgeColor = FLinearColor(0.2f, 0.5f, 1.0f); // Blue for added files
+		break;
+	case ESandboxChangeType::Modified:
+		IconBrush = FAppStyle::GetBrush("Icons.Edit");
+		BadgeColor = FLinearColor(0.2f, 1.0f, 0.2f); // Green for modified files
+		break;
+	case ESandboxChangeType::Deleted:
+		IconBrush = FAppStyle::GetBrush("Icons.Delete");
+		BadgeColor = FLinearColor(1.0f, 0.2f, 0.2f); // Red for deleted files
+		break;
+	default:
+		return SNullWidget::NullWidget;
+	}
+
+	return SNew(SImage)
+		.Image(IconBrush)
+		.ColorAndOpacity(BadgeColor);
+}
+
+TSharedRef<SWidget> FContentBrowserBadgeFeature::GenerateBadgeTooltip(const FAssetData& AssetData)
+{
+	const ESandboxChangeType ChangeType = GetAssetChangeType(AssetData);
+
+	if (ChangeType == ESandboxChangeType::None)
+	{
+		// No tooltip for unchanged assets
+		return SNullWidget::NullWidget;
+	}
+
+	// Create descriptive text based on change type
+	FText TooltipText;
+	switch (ChangeType)
+	{
+	case ESandboxChangeType::Added:
+		TooltipText = FText::Format(
+			NSLOCTEXT("SandboxedEditing", "BadgeTooltip_Added", "Added by sandbox: {0}"),
+			FText::FromString(SandboxModel->GetActiveSandboxName())
+		);
+		break;
+	case ESandboxChangeType::Modified:
+		TooltipText = FText::Format(
+			NSLOCTEXT("SandboxedEditing", "BadgeTooltip_Modified", "Modified by sandbox: {0}"),
+			FText::FromString(SandboxModel->GetActiveSandboxName())
+		);
+		break;
+	case ESandboxChangeType::Deleted:
+		TooltipText = FText::Format(
+			NSLOCTEXT("SandboxedEditing", "BadgeTooltip_Deleted", "Deleted in sandbox: {0}"),
+			FText::FromString(SandboxModel->GetActiveSandboxName())
+		);
+		break;
+	default:
+		return SNullWidget::NullWidget;
+	}
+
+	return SNew(STextBlock)
+		.Text(TooltipText);
+}
+
+FContentBrowserBadgeFeature::ESandboxChangeType FContentBrowserBadgeFeature::GetAssetChangeType(const FAssetData& AssetData) const
+{
+	using namespace FileSandboxCore;
+
+	if (!SandboxModel->HasActiveSandbox())
+	{
+		return ESandboxChangeType::None;
+	}
+
+	// Get the package file path from the asset data
+	FString PackageFilePath = FPackageName::LongPackageNameToFilename(
+		AssetData.PackageName.ToString(),
+		FPackageName::GetAssetPackageExtension()
+	);
+
+	// Convert to absolute path for comparison with sandbox change paths (which are absolute)
+	PackageFilePath = FPaths::ConvertRelativePathToFull(PackageFilePath);
+
+	// Check if this asset's file is in the change list by enumerating changes
+	ESandboxChangeType ResultChangeType = ESandboxChangeType::None;
+
+	SandboxModel->EnumerateFileChanges(
+		SandboxModel->GetActiveSandboxPath(),
+		[&PackageFilePath, &ResultChangeType](const FSandboxedFileChangeInfo& Change) -> UE::FileSandboxCore::EBreakBehavior
+		{
+			if (FPaths::IsSamePath(Change.Path, PackageFilePath))
+			{
+				switch (Change.Action)
+				{
+				case ESandboxFileChange::Added:
+					ResultChangeType = ESandboxChangeType::Added;
+					break;
+				case ESandboxFileChange::Edited:
+					ResultChangeType = ESandboxChangeType::Modified;
+					break;
+				case ESandboxFileChange::Removed:
+					ResultChangeType = ESandboxChangeType::Deleted;
+					break;
+				default:
+					break;
+				}
+				return UE::FileSandboxCore::EBreakBehavior::Break; // Found the file, stop searching
+			}
+			return UE::FileSandboxCore::EBreakBehavior::Continue;
+		}
+	);
+
+	return ResultChangeType;
+}
+
+} // namespace UE::SandboxedEditing
diff --git a/Engine/Plugins/Developer/Sandbox/SandboxedEditing/Source/SandboxedEditing/Private/Features/ContentBrowserBadge/ContentBrowserBadgeFeature.h b/Engine/Plugins/Developer/Sandbox/SandboxedEditing/Source/SandboxedEditing/Private/Features/ContentBrowserBadge/ContentBrowserBadgeFeature.h
new file mode 100644
--- /dev/null
+++ b/Engine/Plugins/Developer/Sandbox/SandboxedEditing/Source/SandboxedEditing/Private/Features/ContentBrowserBadge/ContentBrowserBadgeFeature.h
@@ -0,0 +1,71 @@
+// Copyright Epic Games, Inc. All Rights Reserved.
+
+#pragma once
+
+#include "Delegates/IDelegateInstance.h"
+#include "HAL/Platform.h"
+#include "Templates/SharedPointer.h"
+
+class SWidget;
+struct FAssetData;
+
+namespace UE::SandboxedEditing
+{
+	
+class FSandboxSystemModel;
+
+/**
+ * Feature that displays badges on Content Browser items for files modified in the active sandbox.
+ * Registers badge generators with the Content Browser to show visual indicators for:
+ * - Added files (new files created in sandbox)
+ * - Modified files (existing files changed in sandbox)
+ * - Deleted files (files removed in sandbox)
+ */
+class FContentBrowserBadgeFeature : public FNoncopyable
+{
+public:
+	explicit FContentBrowserBadgeFeature(TSharedRef<FSandboxSystemModel> InSandboxModel);
+	~FContentBrowserBadgeFeature();
+
+private:
+	/** Registers the badge generators with Content Browser */
+	void RegisterBadgeGenerators();
+
+	/** Unregisters the badge generators from Content Browser */
+	void UnregisterBadgeGenerators();
+
+	/** Called when a sandbox is loaded - registers badge generators */
+	void OnSandboxLoaded();
+
+	/** Called when a sandbox is left - unregisters badge generators */
+	void OnSandboxLeft();
+
+	/** Generates the badge icon widget for an asset */
+	TSharedRef<SWidget> GenerateBadgeIcon(const FAssetData& AssetData);
+
+	/** Generates the badge tooltip widget for an asset */
+	TSharedRef<SWidget> GenerateBadgeTooltip(const FAssetData& AssetData);
+
+	/** Checks if the given asset is modified in the active sandbox and returns the change type */
+	enum class ESandboxChangeType : uint8
+	{
+		None,
+		Added,
+		Modified,
+		Deleted
+	};
+	ESandboxChangeType GetAssetChangeType(const FAssetData& AssetData) const;
+
+private:
+	/** Model for interacting with the sandbox system */
+	TSharedRef<FSandboxSystemModel> SandboxModel;
+
+	/** Handle to the registered badge generator */
+	FDelegateHandle BadgeGeneratorHandle;
+
+	/** Delegate handles for sandbox lifecycle events */
+	FDelegateHandle OnLoadSandboxHandle;
+	FDelegateHandle OnLeaveSandboxHandle;
+};
+
+} // namespace UE::SandboxedEditing
diff --git a/Engine/Plugins/Developer/Sandbox/SandboxedEditing/Source/SandboxedEditing/Private/Framework/SandboxedEditingApp.cpp b/Engine/Plugins/Developer/Sandbox/SandboxedEditing/Source/SandboxedEditing/Private/Framework/SandboxedEditingApp.cpp
--- a/Engine/Plugins/Developer/Sandbox/SandboxedEditing/Source/SandboxedEditing/Private/Framework/SandboxedEditingApp.cpp
+++ b/Engine/Plugins/Developer/Sandbox/SandboxedEditing/Source/SandboxedEditing/Private/Framework/SandboxedEditingApp.cpp
@@ -14,6 +14,7 @@ FSandboxedEditingApp::FSandboxedEditingApp()
 	, SandboxIntegration(MakeShared<FSandboxedEditingSandboxIntegration>())
 	, StatusBarFeature(SandboxModel)
 	, BrowserFeature(SandboxModel, SandboxIntegration->GetExternalSandboxViewModel())
+	, ContentBrowserBadgeFeature(SandboxModel)
 	, StartupBehaviorHandler(SandboxModel)
 	, ShutdownCleanupHandler(SandboxModel)
 {
diff --git a/Engine/Plugins/Developer/Sandbox/SandboxedEditing/Source/SandboxedEditing/Private/Framework/SandboxedEditingApp.h b/Engine/Plugins/Developer/Sandbox/SandboxedEditing/Source/SandboxedEditing/Private/Framework/SandboxedEditingApp.h
--- a/Engine/Plugins/Developer/Sandbox/SandboxedEditing/Source/SandboxedEditing/Private/Framework/SandboxedEditingApp.h
+++ b/Engine/Plugins/Developer/Sandbox/SandboxedEditing/Source/SandboxedEditing/Private/Framework/SandboxedEditingApp.h
@@ -3,6 +3,7 @@
 #pragma once
 
 #include "Features/Browser/BrowserFeature.h"
+#include "Features/ContentBrowserBadge/ContentBrowserBadgeFeature.h"
 #include "Features/StatusBar/StatusBarFeature.h"
 #include "Framework/Cleanup/ShutdownCleanupHandler.h"
 #include "Framework/StartupBehavior/StartupBehaviorHandler.h"
@@ -39,6 +40,8 @@ private:
 	const FStatusBarFeature StatusBarFeature;
 	/** Extends the editor with a nomad tab containing a browser for sandboxes. */
 	const FBrowserFeature BrowserFeature;
+	/** Displays badges on Content Browser items for files modified in the active sandbox. */
+	const FContentBrowserBadgeFeature ContentBrowserBadgeFeature;
 	/** Handles startup behavior (creating/loading sandboxes on editor startup). */
 	const FStartupBehaviorHandler StartupBehaviorHandler;
 	/** Handles cleanup of empty sandboxes on editor shutdown. */
diff --git a/Engine/Plugins/Developer/Sandbox/SandboxedEditing/Source/SandboxedEditing/SandboxedEditing.Build.cs b/Engine/Plugins/Developer/Sandbox/SandboxedEditing/Source/SandboxedEditing/SandboxedEditing.Build.cs
--- a/Engine/Plugins/Developer/Sandbox/SandboxedEditing/Source/SandboxedEditing/SandboxedEditing.Build.cs
+++ b/Engine/Plugins/Developer/Sandbox/SandboxedEditing/Source/SandboxedEditing/SandboxedEditing.Build.cs
@@ -18,6 +18,7 @@ namespace UnrealBuildTool.Rules
 				new string[]
 				{
 					"AssetTools",
+					"ContentBrowser",
 					"DesktopPlatform",
 					"Engine",
 					"FileSandboxCore",
